- Minor cleanup to Turnkey menu

diff --git a/Engine/Source/Editor/TurnkeySupport/Private/TurnkeySupportModule.cpp b/Engine/Source/Editor/TurnkeySupport/Private/TurnkeySupportModule.cpp
--- a/Engine/Source/Editor/TurnkeySupport/Private/TurnkeySupportModule.cpp
+++ b/Engine/Source/Editor/TurnkeySupport/Private/TurnkeySupportModule.cpp
@@ -196,7 +196,8 @@ public:
 
 	static bool CanCookOrPackage(FName IniPlatformName, EPrepareContentMode Mode)
 	{
-		return GetTargetPlatformManager()->FindTargetPlatform(IniPlatformName.ToString()) != nullptr;
+		// prep for debugging not supported yet
+		return Mode != EPrepareContentMode::PrepareForDebugging && GetTargetPlatformManager()->FindTargetPlatform(IniPlatformName.ToString()) != nullptr;
 	}
 
 	static void CookOrPackage(FName IniPlatformName, EPrepareContentMode Mode)
@@ -666,6 +667,15 @@ static void MakeTurnkeyPlatformMenu(FMenuBuilder& MenuBuilder, FName IniPlatform
 
 		MenuBuilder.BeginSection("ContentManagement", LOCTEXT("TurnkeySection_Content", "Content Management"));
 
+		MenuBuilder.AddMenuEntry(
+			LOCTEXT("Turnkey_PackageProject", "Package Project"),
+			LOCTEXT("TurnkeyTooltip_PackageProject", "Package this project and archive it to a user-selected directory. This can then be used to install and run."),
+			FSlateIcon(),
+			FUIAction(
+				FExecuteAction::CreateStatic(&FTurnkeySupportCallbacks::CookOrPackage, IniPlatformName, EPrepareContentMode::Package),
+				FCanExecuteAction::CreateStatic(&FTurnkeySupportCallbacks::CanCookOrPackage, IniPlatformName, EPrepareContentMode::Package)
+			)
+		);
 
 		MenuBuilder.AddMenuEntry(
 			LOCTEXT("Turnkey_CookContent", "Cook Content"),
@@ -687,16 +697,6 @@ static void MakeTurnkeyPlatformMenu(FMenuBuilder& MenuBuilder, FName IniPlatform
 			)
 		);
 
-		MenuBuilder.AddMenuEntry(
-			LOCTEXT("Turnkey_PackageProject", "Package Project"),
-			LOCTEXT("TurnkeyTooltip_PackageProject", "Package this project and archive it to a user-selected directory. This can then be used to install and run."),
-			FSlateIcon(),
-			FUIAction(
-				FExecuteAction::CreateStatic(&FTurnkeySupportCallbacks::CookOrPackage, IniPlatformName, EPrepareContentMode::Package),
-				FCanExecuteAction::CreateStatic(&FTurnkeySupportCallbacks::CanCookOrPackage, IniPlatformName, EPrepareContentMode::Package)
-			)
-		);
-
 		MenuBuilder.EndSection();
 
 
@@ -1195,6 +1195,8 @@ TSharedRef<SWidget> FTurnkeySupportModule::MakeTurnkeyMenu() const
 
 	if (UncompiledPlatforms.Num() != 0)
 	{
+		MenuBuilder.AddSeparator(NAME_None);
+
 		MenuBuilder.AddSubMenu(
 			LOCTEXT("Turnkey_UncompiledPlatforms", "Platforms With No Compiled Support"),
 			LOCTEXT("Turnkey_UncompiledPlatformsToolTip", "List of platforms that you have access to, but support is not compiled in to the editor. It may be caused by missing an SDK, so you attempt to install an SDK here."),
