- Refactored more Turnkey stuff into TurnkeySupportModule (moved stuff out of DDPI, etc)
- Also fixed a crash or two

diff --git a/Engine/Source/Developer/IOS/IOSPlatformEditor/IOSPlatformEditor.Build.cs b/Engine/Source/Developer/IOS/IOSPlatformEditor/IOSPlatformEditor.Build.cs
--- a/Engine/Source/Developer/IOS/IOSPlatformEditor/IOSPlatformEditor.Build.cs
+++ b/Engine/Source/Developer/IOS/IOSPlatformEditor/IOSPlatformEditor.Build.cs
@@ -35,6 +35,7 @@ public class IOSPlatformEditor : ModuleRules
 			new string[] {
 				"GameProjectGeneration",
 				"Settings",
+				"TurnkeySupport",
 			}
 		);
 	}
diff --git a/Engine/Source/Developer/IOS/IOSPlatformEditor/Private/IOSCustomIconProjectBuildMutatorFeature.cpp b/Engine/Source/Developer/IOS/IOSPlatformEditor/Private/IOSCustomIconProjectBuildMutatorFeature.cpp
--- a/Engine/Source/Developer/IOS/IOSPlatformEditor/Private/IOSCustomIconProjectBuildMutatorFeature.cpp
+++ b/Engine/Source/Developer/IOS/IOSPlatformEditor/Private/IOSCustomIconProjectBuildMutatorFeature.cpp
@@ -8,6 +8,7 @@
 #include "Interfaces/ITargetPlatformManagerModule.h"
 #include "GenericPlatform/GenericPlatformFile.h"
 #include "Misc/CoreMisc.h"
+#include "Interfaces/IturnkeySupportModule.h"
 
 static bool RequiresBuild()
 {
@@ -50,7 +51,7 @@ bool FIOSCustomIconProjectBuildMutatorFeature ::RequiresProjectBuild(const FName
 	const PlatformInfo::FTargetPlatformInfo* const PlatInfo = PlatformInfo::FindPlatformInfo(InPlatformInfoName);
 	check(PlatInfo);
 
-	if (PlatInfo->DataDrivenPlatformInfo->GetSdkStatus() == DDPIPlatformSdkStatus::Valid)
+	if (ITurnkeySupportModule::Get().GetSdkInfo(PlatInfo->IniPlatformName).Status == ETurnkeyPlatformSdkStatus::Valid)
 	{
 		const ITargetPlatform* const Platform = GetTargetPlatformManager()->FindTargetPlatform(PlatInfo->Name.ToString());
 		if (Platform)
diff --git a/Engine/Source/Developer/TargetPlatform/Private/TargetPlatformManagerModule.cpp b/Engine/Source/Developer/TargetPlatform/Private/TargetPlatformManagerModule.cpp
--- a/Engine/Source/Developer/TargetPlatform/Private/TargetPlatformManagerModule.cpp
+++ b/Engine/Source/Developer/TargetPlatform/Private/TargetPlatformManagerModule.cpp
@@ -20,6 +20,7 @@
 #include "Interfaces/ITextureFormatModule.h"
 #include "PlatformInfo.h"
 #include "DesktopPlatformModule.h"
+#include "Interfaces/ITurnkeySupportModule.h"
 
 #if PHYSICS_INTERFACE_PHYSX
 #include "IPhysXCooking.h"
@@ -88,7 +89,7 @@ public:
 		, bHasInitErrors(false)
 		, bIgnoreFirstDelegateCall(true)
 	{
-		FDataDrivenPlatformInfoRegistry::UpdateSdkStatus();
+		ITurnkeySupportModule::Get().UpdateSdkInfo();
 
 #if AUTOSDKS_ENABLED		
 		
@@ -1134,10 +1135,10 @@ protected:
 
 		if (bTPInitialized)
 		{
-			FDataDrivenPlatformInfoRegistry::UpdateSdkStatus();
+			ITurnkeySupportModule::Get().UpdateSdkInfo();
 		}
 
-		FDataDrivenPlatformInfoRegistry::ClearDeviceStatus(PlatformName);
+		ITurnkeySupportModule::Get().ClearDeviceStatus(PlatformName);
 		Invalidate();
 
 		return bTPInitialized;
diff --git a/Engine/Source/Developer/TargetPlatform/TargetPlatform.Build.cs b/Engine/Source/Developer/TargetPlatform/TargetPlatform.Build.cs
--- a/Engine/Source/Developer/TargetPlatform/TargetPlatform.Build.cs
+++ b/Engine/Source/Developer/TargetPlatform/TargetPlatform.Build.cs
@@ -18,6 +18,7 @@ public class TargetPlatform : ModuleRules
 
 		PrivateIncludePathModuleNames.Add("Engine");
 		PrivateIncludePathModuleNames.Add("PhysicsCore");
+		PrivateIncludePathModuleNames.Add("TurnkeySupport");
 
 		// no need for all these modules if the program doesn't want developer tools at all (like UnrealFileServer)
 		if (!Target.bBuildRequiresCookedData && Target.bBuildDeveloperTools)
diff --git a/Engine/Source/Editor/TurnkeySupport/Private/TurnkeySupportModule.cpp b/Engine/Source/Editor/TurnkeySupport/Private/TurnkeySupportModule.cpp
--- a/Engine/Source/Editor/TurnkeySupport/Private/TurnkeySupportModule.cpp
+++ b/Engine/Source/Editor/TurnkeySupport/Private/TurnkeySupportModule.cpp
@@ -1,6 +1,7 @@
 // Copyright Epic Games, Inc. All Rights Reserved.
 
 #include "TurnkeySupportModule.h"
+#include "TurnkeySupport.h"
 
 #include "SlateOptMacros.h"
 #include "Widgets/Images/SImage.h"
@@ -46,12 +47,19 @@
 #include "Framework/MultiBox/MultiBoxBuilder.h"
 #include "Developer/DerivedDataCache/Public/DerivedDataCacheInterface.h"
 
+#include "Misc/MonitoredProcess.h"
+
 
 
 
 DEFINE_LOG_CATEGORY(LogTurnkeySupport);
 #define LOCTEXT_NAMESPACE "FTurnkeySupportModule"
 
+namespace 
+{
+	FCriticalSection GTurnkeySection;
+}
+
 
 enum class EPrepareContentMode : uint8
 {
@@ -241,7 +249,7 @@ public:
 				return;
 			}
 
-			if (PlatformInfo->DataDrivenPlatformInfo->GetSdkStatus() != DDPIPlatformSdkStatus::Valid && ShowBadSDKDialog(IniPlatformName) == false)
+			if (ITurnkeySupportModule::Get().GetSdkInfo(IniPlatformName).Status != ETurnkeyPlatformSdkStatus::Valid && ShowBadSDKDialog(IniPlatformName) == false)
 			{
 				return;
 			}
@@ -654,17 +662,16 @@ static TAttribute<FText> MakeSdkStatusAttribute(FName IniPlatformName, TSharedPt
 	return TAttribute<FText>::Create(TAttribute<FText>::FGetter::CreateLambda([IniPlatformName, DisplayString, DeviceId]()
 	{
 		// get the status, or Unknown if it's not there
-		const FDataDrivenPlatformInfo& Info = FDataDrivenPlatformInfoRegistry::GetPlatformInfo(IniPlatformName);
-		DDPIPlatformSdkStatus Status = DeviceId.Len() ? Info.GetStatusForDeviceId(DeviceId) : Info.GetSdkStatus(false);
+		ETurnkeyPlatformSdkStatus Status = DeviceId.Len() ? ITurnkeySupportModule::Get().GetSdkInfoForDeviceId(DeviceId).Status : ITurnkeySupportModule::Get().GetSdkInfo(IniPlatformName, false).Status;
 
 		// @todo turnkey: Have premade FText's by SdkStatus for speed
 		const TCHAR* Desc =
-			Status == DDPIPlatformSdkStatus::Querying ? TEXT("Querying...") :
-			Status == DDPIPlatformSdkStatus::Valid ? TEXT("Valid Sdk") :
-			Status == DDPIPlatformSdkStatus::OutOfDate ? TEXT("Outdated Sdk") :
-			Status == DDPIPlatformSdkStatus::NoSdk ? TEXT("No Sdk") :
-			Status == DDPIPlatformSdkStatus::FlashValid ? TEXT("Valid Flash") :
-			Status == DDPIPlatformSdkStatus::FlashOutOfDate ? TEXT("Outdated Flash") :
+			Status == ETurnkeyPlatformSdkStatus::Querying ? TEXT("Querying...") :
+			Status == ETurnkeyPlatformSdkStatus::Valid ? TEXT("Valid Sdk") :
+			Status == ETurnkeyPlatformSdkStatus::OutOfDate ? TEXT("Outdated Sdk") :
+			Status == ETurnkeyPlatformSdkStatus::NoSdk ? TEXT("No Sdk") :
+			Status == ETurnkeyPlatformSdkStatus::FlashValid ? TEXT("Valid Flash") :
+			Status == ETurnkeyPlatformSdkStatus::FlashOutOfDate ? TEXT("Outdated Flash") :
 			TEXT("???");
 		return FText::FromString(FString::Printf(TEXT("%s (%s)"), *DisplayString, Desc));
 	}));
@@ -782,8 +789,7 @@ static void MakeTurnkeyPlatformMenu(FMenuBuilder& MenuBuilder, FName IniPlatform
 				FText(),
 				FNewMenuDelegate::CreateLambda([UBTPlatformString, IniPlatformName, DeviceName, DeviceId](FMenuBuilder& SubMenuBuilder)
 				{
-					const FDataDrivenPlatformInfo& DDPI = FDataDrivenPlatformInfoRegistry::GetPlatformInfo(IniPlatformName);
-					const FDDPISdkInfo& SdkInfo = DDPI.GetSdkInfoForDeviceId(DeviceId);
+					FTurnkeySdkInfo SdkInfo = ITurnkeySupportModule::Get().GetSdkInfoForDeviceId(DeviceId);
 					FFormatOrderedArguments Args = { FText::FromString(SdkInfo.InstalledVersion), FText::FromString(SdkInfo.MinAllowedVersion),FText::FromString(SdkInfo.MaxAllowedVersion) };
 					SubMenuBuilder.AddWidget(
 						SNew(STextBlock)
@@ -815,7 +821,7 @@ static void MakeTurnkeyPlatformMenu(FMenuBuilder& MenuBuilder, FName IniPlatform
 
 	MenuBuilder.BeginSection("SdkManagement", LOCTEXT("TurnkeySection_Sdks", "Sdk Managment"));
 
-	const FDDPISdkInfo& SdkInfo = DDPI.GetSdkInfo(true);
+	const FTurnkeySdkInfo& SdkInfo = ITurnkeySupportModule::Get().GetSdkInfo(IniPlatformName, true);
 	FFormatOrderedArguments Args = { FText::FromString(SdkInfo.InstalledVersion), FText::FromString(SdkInfo.AutoSDKVersion),FText::FromString(SdkInfo.MinAllowedVersion),FText::FromString(SdkInfo.MaxAllowedVersion) };
 	MenuBuilder.AddWidget(
 		SNew(STextBlock)
@@ -825,7 +831,7 @@ static void MakeTurnkeyPlatformMenu(FMenuBuilder& MenuBuilder, FName IniPlatform
 	);
 
 	const TCHAR* NoDevice = nullptr;
-	if (DDPI.GetSdkStatus() == DDPIPlatformSdkStatus::OutOfDate)
+	if (SdkInfo.Status == ETurnkeyPlatformSdkStatus::OutOfDate)
 	{
 		MenuBuilder.AddMenuEntry(
 			LOCTEXT("Turnkey_InstallSdkMinimal", "Update Sdk (Prefer Minimal)"),
@@ -841,7 +847,7 @@ static void MakeTurnkeyPlatformMenu(FMenuBuilder& MenuBuilder, FName IniPlatform
 			FExecuteAction::CreateStatic(TurnkeyInstallSdk, UBTPlatformString, true, false, NoDevice)
 		);
 	}
-	else if (DDPI.GetSdkStatus() == DDPIPlatformSdkStatus::Valid)
+	else if (SdkInfo.Status == ETurnkeyPlatformSdkStatus::Valid)
 	{
 		MenuBuilder.AddMenuEntry(
 			LOCTEXT("Turnkey_InstallSdkMinimal", "Force Reinstall Sdk (Prefer Minimal)"),
@@ -1077,6 +1083,8 @@ static void GenerateDeviceProxyMenuParams(TSharedPtr<ITargetDeviceProxy> DeviceP
 
 TSharedRef<SWidget> FTurnkeySupportModule::MakeTurnkeyMenu() const
 {
+	FTurnkeySupportCommands::Register();
+
 	const FTurnkeySupportCommands& Commands = FTurnkeySupportCommands::Get();
 
 	const bool bShouldCloseWindowAfterMenuSelection = true;
@@ -1161,11 +1169,12 @@ TSharedRef<SWidget> FTurnkeySupportModule::MakeTurnkeyMenu() const
 						);
 				}
 
-				// gather any unknown status devices to query at the end
+				ITurnkeySupportModule& TurnkeySupport = ITurnkeySupportModule::Get();
+				// gath	er any unknown status devices to query at the end
 				for (const TSharedPtr<ITargetDeviceProxy> Proxy : DeviceProxies)
 				{
 					FString DeviceId = Proxy->GetTargetDeviceId(NAME_None);
-					if (Pair.Value.GetStatusForDeviceId(DeviceId) == DDPIPlatformSdkStatus::Unknown)
+					if (TurnkeySupport.GetSdkInfoForDeviceId(DeviceId).Status == ETurnkeyPlatformSdkStatus::Unknown)
 					{
 						DeviceIdsToQuery.Add(DeviceId);
 					}
@@ -1204,7 +1213,7 @@ TSharedRef<SWidget> FTurnkeySupportModule::MakeTurnkeyMenu() const
 
 	MenuBuilder.BeginSection("AllPlatforms", LOCTEXT("TurnkeyMenu_ManagePlatforms", "Content/Sdk/Device Management"));
 	TMap<FName, const FDataDrivenPlatformInfo*> UncompiledPlatforms;
-	for (const auto Pair : FDataDrivenPlatformInfoRegistry::GetAllPlatformInfos())
+	for (const auto& Pair : FDataDrivenPlatformInfoRegistry::GetAllPlatformInfos())
 	{
 		if (Pair.Value.bIsFakePlatform || Pair.Value.bEnabledForUse == false)
 		{
@@ -1303,18 +1312,348 @@ TSharedRef<SWidget> FTurnkeySupportModule::MakeTurnkeyMenu() const
 	// now kick-off any devices that need to be updated
 	if (DeviceIdsToQuery.Num() > 0)
 	{
-		FDataDrivenPlatformInfoRegistry::UpdateDeviceSdkStatus(DeviceIdsToQuery);
+		ITurnkeySupportModule::Get().UpdateSdkInfoForDevices(DeviceIdsToQuery);
 	}
 
 	return MenuBuilder.MakeWidget();
 }
 
-/* IModuleInterface implementation
- *****************************************************************************/
+
+
+
+
+
+// some shared functionality
+static void PrepForTurnkeyReport(FString& Command, FString& BaseCommandline, FString& ReportFilename)
+{
+	static int ReportIndex = 0;
+
+	FString LogFilename = FPaths::ConvertRelativePathToFull(FPaths::Combine(FPaths::ProjectIntermediateDir(), *FString::Printf(TEXT("TurnkeyLog_%d.log"), ReportIndex)));
+	ReportFilename = FPaths::ConvertRelativePathToFull(FPaths::Combine(FPaths::ProjectIntermediateDir(), *FString::Printf(TEXT("TurnkeyReport_%d.log"), ReportIndex++)));
+
+	// make sure intermediate directory exists
+	IFileManager::Get().MakeDirectory(*FPaths::ProjectIntermediateDir());
+
+	Command = TEXT("{EngineDir}Build/BatchFiles/RunuAT");
+	//	Command = TEXT("{EngineDir}/Binaries/DotNET/AutomationTool.exe");
+	BaseCommandline = FString::Printf(TEXT("Turnkey -utf8output -WaitForUATMutex -command=VerifySdk -ReportFilename=\"%s\" -log=\"%s\""), *ReportFilename, *LogFilename);
+
+	// convert into appropriate calls for the current platform
+	FPlatformProcess::ModifyCreateProcParams(Command, BaseCommandline, FGenericPlatformProcess::ECreateProcHelperFlags::AppendPlatformScriptExtension | FGenericPlatformProcess::ECreateProcHelperFlags::RunThroughShell);
+}
+
+bool GetSdkInfoFromTurnkey(FString Line, FName& PlatformName, FString& DeviceId, FTurnkeySdkInfo& SdkInfo)
+{
+	int32 Colon = Line.Find(TEXT(": "));
+
+	if (Colon < 0)
+	{
+		return false;
+	}
+
+	// break up the string
+	FString PlatformString = Line.Mid(0, Colon);
+	FString Info = Line.Mid(Colon + 2);
+
+	int32 AtSign = PlatformString.Find(TEXT("@"));
+	if (AtSign > 0)
+	{
+		// return the platform@name as the deviceId, then remove the @name part for the platform
+		DeviceId = ConvertToDDPIDeviceId(PlatformString);
+		PlatformString = PlatformString.Mid(0, AtSign);
+	}
+
+	// get the DDPI name
+	PlatformName = FName(*ConvertToDDPIPlatform(PlatformString));
+
+	// parse out the results from the (key=val, key=val) result from turnkey
+	FString StatusString;
+	FString FlagsString;
+	FParse::Value(*Info, TEXT("Status="), StatusString);
+	FParse::Value(*Info, TEXT("Flags="), FlagsString);
+	FParse::Value(*Info, TEXT("Installed="), SdkInfo.InstalledVersion);
+	FParse::Value(*Info, TEXT("AutoSDK="), SdkInfo.AutoSDKVersion);
+	FParse::Value(*Info, TEXT("MinAllowed="), SdkInfo.MinAllowedVersion);
+	FParse::Value(*Info, TEXT("MaxAllowed="), SdkInfo.MaxAllowedVersion);
+
+	SdkInfo.Status = ETurnkeyPlatformSdkStatus::Unknown;
+	if (StatusString == TEXT("Valid"))
+	{
+		SdkInfo.Status = ETurnkeyPlatformSdkStatus::Valid;
+	}
+	else
+	{
+		if (FlagsString.Contains(TEXT("AutoSdk_InvalidVersionExists")) || FlagsString.Contains(TEXT("InstalledSdk_InvalidVersionExists")))
+		{
+			SdkInfo.Status = ETurnkeyPlatformSdkStatus::OutOfDate;
+		}
+		else
+		{
+			SdkInfo.Status = ETurnkeyPlatformSdkStatus::NoSdk;
+		}
+	}
+
+	return true;
+}
+
+// FDataDrivenPlatformInfo& DeviceIdToInfo(FString DeviceId, FString* OutDeviceName)
+// {
+// 	TArray<FString> PlatformAndDevice;
+// 	DeviceId.ParseIntoArray(PlatformAndDevice, TEXT("@"), true);
+// 
+// 	if (OutDeviceName)
+// 	{
+// 		*OutDeviceName = PlatformAndDevice[1];
+// 	}
+// 
+// 	FString DDPIPlatformName = ConvertToDDPIPlatform(PlatformAndDevice[0]);
+// 
+// 	checkf(DataDrivenPlatforms.Contains(*DDPIPlatformName), TEXT("DataDrivenPlatforms map did not contain the DDPI Platform %s"), *DDPIPlatformName);
+// 
+// 	// have to convert back to Windows from Win64
+// 	return DataDrivenPlatforms[*DDPIPlatformName];
+// 
+// }
+
+
+void FTurnkeySupportModule::UpdateSdkInfo()
+{
+	// make sure all known platforms are in the map
+	if (PerPlatformSdkInfo.Num() == 0)
+	{
+		for (auto& It : FDataDrivenPlatformInfoRegistry::GetAllPlatformInfos())
+		{
+			PerPlatformSdkInfo.Add(It.Key, FTurnkeySdkInfo());
+		}
+	}
+
+	// don't run UAT from commandlets (like the cooker) that are often launched from UAT and this will go poorly
+	if (IsRunningCommandlet())
+	{
+		return;
+	}
+
+
+	FString Command, BaseCommandline, ReportFilename;
+	PrepForTurnkeyReport(Command, BaseCommandline, ReportFilename);
+	// get status for all platforms
+	FString Commandline = BaseCommandline + TEXT(" -platform=all");
+
+	UE_LOG(LogInit, Log, TEXT("Running Turnkey SDK detection: '%s %s'"), *Command, *Commandline);
+
+	{
+		FScopeLock Lock(&GTurnkeySection);
+
+		// reset status to unknown
+		for (auto& It : PerPlatformSdkInfo)
+		{
+			It.Value.Status = ETurnkeyPlatformSdkStatus::Querying;
+		}
+
+		// reset the per-device status when querying general Sdk status
+		ClearDeviceStatus();
+	}
+
+	FMonitoredProcess* TurnkeyProcess = new FMonitoredProcess(Command, Commandline, true, false);
+	TurnkeyProcess->OnCompleted().BindLambda([this, ReportFilename, TurnkeyProcess](int32 ExitCode)
+	{
+		AsyncTask(ENamedThreads::GameThread, [this, ReportFilename, TurnkeyProcess, ExitCode]()
+		{
+			FScopeLock Lock(&GTurnkeySection);
+
+			if (ExitCode == 0 || ExitCode == 10)
+			{
+				TArray<FString> Contents;
+				if (FFileHelper::LoadFileToStringArray(Contents, *ReportFilename))
+				{
+					for (FString& Line : Contents)
+					{
+						UE_LOG(LogTemp, Log, TEXT("Turnkey Platform: %s"), *Line);
+
+						// parse a Turnkey line
+						FName PlatformName;
+						FString Unused;
+						FTurnkeySdkInfo SdkInfo;
+						if (GetSdkInfoFromTurnkey(Line, PlatformName, Unused, SdkInfo) == false)
+						{
+							continue;
+						}
+
+						// check if we had already set a ManualSDK - and don't set it again. Because of the way AutoSDKs are activated in the editor after the first call to Turnkey,
+						// future calls to Turnkey will inherit the AutoSDK env vars, and it won't be able to determine the manual SDK versions anymore. If we use the editor to
+						// install an SDK via Turnkey, it will directly update the installed version based on the result of that command, not this Update operation
+
+						FString OriginalManualInstallValue = PerPlatformSdkInfo[PlatformName].InstalledVersion;
+
+						// set it into the platform
+						PerPlatformSdkInfo[PlatformName] = SdkInfo;
+
+						// restore the original installed version if it set after the first time
+						if (OriginalManualInstallValue.Len() > 0)
+						{
+							PerPlatformSdkInfo[PlatformName].InstalledVersion = OriginalManualInstallValue;
+						}
+
+
+						UE_LOG(LogTemp, Log, TEXT("[TEST] Turnkey Platform: %s - %d, Installed: %s, AudoSDK: %s, Allowed: %s-%s"), *PlatformName.ToString(), (int)SdkInfo.Status, *SdkInfo.InstalledVersion,
+							*SdkInfo.AutoSDKVersion, *SdkInfo.MinAllowedVersion, *SdkInfo.MaxAllowedVersion);
+					}
+				}
+			}
+			else
+			{
+				for (auto& It : PerPlatformSdkInfo)
+				{
+					It.Value.Status = ETurnkeyPlatformSdkStatus::Error;
+					It.Value.SdkErrorInformation = FText::Format(NSLOCTEXT("Turnkey", "TurnkeyError_ReturnedError", "Turnkey returned an error, code {0}"), { ExitCode });
+
+					// @todo turnkey error description!
+				}
+			}
+
+
+			for (auto& It : PerPlatformSdkInfo)
+			{
+				if (It.Value.Status == ETurnkeyPlatformSdkStatus::Querying)
+				{
+					// fake platforms won't come back, just skip it
+					if (FDataDrivenPlatformInfoRegistry::GetPlatformInfo(It.Key).bIsFakePlatform)
+					{
+						It.Value.Status = ETurnkeyPlatformSdkStatus::Unknown;
+					}
+					else
+					{
+						It.Value.Status = ETurnkeyPlatformSdkStatus::Error;
+						It.Value.SdkErrorInformation = NSLOCTEXT("Turnkey", "TurnkeyError_NotReturned", "The platform's Sdk status was not returned from Turnkey");
+					}
+				}
+			}
+
+			// cleanup
+			delete TurnkeyProcess;
+			IFileManager::Get().Delete(*ReportFilename);
+		});
+	});
+
+	// run it
+	TurnkeyProcess->Launch();
+}
+
+void FTurnkeySupportModule::UpdateSdkInfoForDevices(TArray<FString> PlatformDeviceIds)
+{
+	FString Command, BaseCommandline, ReportFilename;
+	PrepForTurnkeyReport(Command, BaseCommandline, ReportFilename);
+
+	// the platform part of the Id may need to be converted to be turnkey (ie UBT) proper
+
+	FString Commandline = BaseCommandline + FString(TEXT(" -Device=")) + FString::JoinBy(PlatformDeviceIds, TEXT("+"), [](FString Id) { return ConvertToUATDeviceId(Id); });
+
+	UE_LOG(LogInit, Log, TEXT("Running Turnkey SDK detection: '%s %s'"), *Command, *Commandline);
+
+	{
+		FScopeLock Lock(&GTurnkeySection);
+
+		// set status to querying
+		FTurnkeySdkInfo DefaultInfo;
+		DefaultInfo.Status = ETurnkeyPlatformSdkStatus::Querying;
+		for (const FString& Id : PlatformDeviceIds)
+		{
+			PerDeviceSdkInfo.Add(ConvertToDDPIDeviceId(Id), DefaultInfo);
+		}
+	}
+
+	FMonitoredProcess* TurnkeyProcess = new FMonitoredProcess(Command, Commandline, true, false);
+	TurnkeyProcess->OnCompleted().BindLambda([this, ReportFilename, TurnkeyProcess, PlatformDeviceIds](int32 ExitCode)
+	{
+		AsyncTask(ENamedThreads::GameThread, [this, ReportFilename, TurnkeyProcess, PlatformDeviceIds, ExitCode]()
+		{
+			FScopeLock Lock(&GTurnkeySection);
+
+			if (ExitCode == 0 || ExitCode == 10)
+			{
+				TArray<FString> Contents;
+				if (FFileHelper::LoadFileToStringArray(Contents, *ReportFilename))
+				{
+					for (FString& Line : Contents)
+					{
+						FName PlatformName;
+						FString DDPIDeviceId;
+						FTurnkeySdkInfo SdkInfo;
+						if (GetSdkInfoFromTurnkey(Line, PlatformName, DDPIDeviceId, SdkInfo) == false)
+						{
+							continue;
+						}
+
+						// skip over non-device lines
+						if (DDPIDeviceId.Len() == 0)
+						{
+							continue;
+						}
+
+						UE_LOG(LogTemp, Log, TEXT("Turnkey Device: %s"), *Line);
+
+						PerDeviceSdkInfo[DDPIDeviceId] = SdkInfo;
+
+						UE_LOG(LogTemp, Log, TEXT("[TEST] Turnkey Device: %s - %d, Installed: %s, Allowed: %s-%s"), *DDPIDeviceId, (int)SdkInfo.Status, *SdkInfo.InstalledVersion,
+							*SdkInfo.MinAllowedVersion, *SdkInfo.MaxAllowedVersion);
+					}
+				}
+			}
+
+			for (const FString& Id : PlatformDeviceIds)
+			{
+				FTurnkeySdkInfo& SdkInfo = PerDeviceSdkInfo[ConvertToDDPIDeviceId(Id)];
+				if (SdkInfo.Status == ETurnkeyPlatformSdkStatus::Querying)
+				{
+					SdkInfo.Status = ETurnkeyPlatformSdkStatus::Error;
+					SdkInfo.SdkErrorInformation = NSLOCTEXT("Turnkey", "TurnkeyError_DeviceNotReturned", "A device's Sdk status was not returned from Turnkey");
+				}
+			}
+
+			// cleanup
+			delete TurnkeyProcess;
+			IFileManager::Get().Delete(*ReportFilename);
+		});
+	});
+
+	// run it
+	TurnkeyProcess->Launch();
+}
+
+FTurnkeySdkInfo FTurnkeySupportModule::GetSdkInfo(FName PlatformName, bool bBlockIfQuerying) const
+{
+	FScopeLock Lock(&GTurnkeySection);
+
+	// return the status, or Unknown info if not known
+	return PerPlatformSdkInfo.FindRef(ConvertToDDPIPlatform(PlatformName));
+}
+
+FTurnkeySdkInfo FTurnkeySupportModule::GetSdkInfoForDeviceId(const FString& DeviceId) const
+{
+	FScopeLock Lock(&GTurnkeySection);
+
+	// return the status, or Unknown info if not known
+	return PerDeviceSdkInfo.FindRef(ConvertToDDPIDeviceId(DeviceId));
+}
+
+void FTurnkeySupportModule::ClearDeviceStatus(FName PlatformName)
+{
+	FScopeLock Lock(&GTurnkeySection);
+
+	FString Prefix = ConvertToDDPIPlatform(PlatformName.ToString()) + "@";
+	for (auto& Pair : PerDeviceSdkInfo)
+	{
+		if (PlatformName == NAME_None || Pair.Key.StartsWith(Prefix))
+		{
+			Pair.Value.Status = ETurnkeyPlatformSdkStatus::Unknown;
+		}
+	}
+}
+
+
 
 void FTurnkeySupportModule::StartupModule( )
 {
-	FTurnkeySupportCommands::Register();
 }
 
 
@@ -1365,7 +1704,7 @@ if (UGameMapsSettings::GetGameDefaultMap().IsEmpty())
 	return;
 }
 
-if (/*PlatformInfo->DataDrivenPlatformInfo->GetSdkStatus() != DDPIPlatformSdkStatus::Valid ||*/ (bProjectHasCode && PlatformInfo->DataDrivenPlatformInfo->bUsesHostCompiler && !FSourceCodeNavigation::IsCompilerAvailable()))
+if (/*PlatformInfo->DataDrivenPlatformInfo->GetSdkStatus() != ETurnkeyPlatformSdkStatus::Valid ||*/ (bProjectHasCode && PlatformInfo->DataDrivenPlatformInfo->bUsesHostCompiler && !FSourceCodeNavigation::IsCompilerAvailable()))
 {
 	if (!ShowBadSDKDialog(PlatformInfo->IniPlatformName))
 	{
diff --git a/Engine/Source/Editor/TurnkeySupport/Private/TurnkeySupportModule.h b/Engine/Source/Editor/TurnkeySupport/Private/TurnkeySupportModule.h
--- a/Engine/Source/Editor/TurnkeySupport/Private/TurnkeySupportModule.h
+++ b/Engine/Source/Editor/TurnkeySupport/Private/TurnkeySupportModule.h
@@ -3,12 +3,9 @@
 #pragma once
 
 #include "CoreMinimal.h"
-#include "Widgets/SWidget.h"
 #include "Interfaces/ITurnkeySupportModule.h"
 
 
-DECLARE_LOG_CATEGORY_EXTERN(LogTurnkeySupport, Log, All);
-
 /**
  * Editor main frame module
  */
@@ -22,6 +19,19 @@ public:
 	 */
 	virtual TSharedRef<SWidget> MakeTurnkeyMenu() const override;
 
+	/**
+	 * @return	The newly-created menu widget
+	 */
+	virtual void UpdateSdkInfo() override;
+
+	/**
+	 * Runs Turnkey to get the Sdk information for a list of devices
+	 */
+	virtual void UpdateSdkInfoForDevices(TArray<FString> DeviceIds) override;
+
+	virtual FTurnkeySdkInfo GetSdkInfo(FName PlatformName, bool bBlockIfQuerying = true) const override;
+	virtual FTurnkeySdkInfo GetSdkInfoForDeviceId(const FString& DeviceId) const override;
+	virtual void ClearDeviceStatus(FName PlatformName=NAME_None) override;
 public:
 
 	// IModuleInterface interface
@@ -36,5 +46,10 @@ public:
 
 
 private:
+	// Information about the validity of using a platform, discovered via Turnkey
+	TMap<FName, FTurnkeySdkInfo> PerPlatformSdkInfo;
+
+	// Information about the validity of each connected device (by string, discovered by Turnkey)
+	TMap<FString, FTurnkeySdkInfo> PerDeviceSdkInfo;
 
 };
diff --git a/Engine/Source/Editor/TurnkeySupport/Public/Interfaces/ITurnkeySupportModule.h b/Engine/Source/Editor/TurnkeySupport/Public/Interfaces/ITurnkeySupportModule.h
--- a/Engine/Source/Editor/TurnkeySupport/Public/Interfaces/ITurnkeySupportModule.h
+++ b/Engine/Source/Editor/TurnkeySupport/Public/Interfaces/ITurnkeySupportModule.h
@@ -3,12 +3,33 @@
 #pragma once
 
 #include "CoreMinimal.h"
-#include "SlateFwd.h"
 #include "Modules/ModuleInterface.h"
 #include "Modules/ModuleManager.h"
-#include "Widgets/SWidget.h"
-#include "Framework/Commands/UICommandList.h"
-#include "Framework/MultiBox/MultiBoxExtender.h"
+
+class SWidget;
+
+enum class ETurnkeyPlatformSdkStatus : uint8
+{
+	Unknown,
+	Querying,
+	Valid,
+	OutOfDate,
+	NoSdk,
+	Error,
+	FlashValid,
+	FlashOutOfDate,
+	// @todo turnkey: add AutoSdkValid and ManualSdkValid, with Valid a Combination of both
+};
+
+struct FTurnkeySdkInfo
+{
+	ETurnkeyPlatformSdkStatus Status = ETurnkeyPlatformSdkStatus::Unknown;
+	FText SdkErrorInformation;
+	FString InstalledVersion;
+	FString AutoSDKVersion; // only valid for platform, not device
+	FString MinAllowedVersion;
+	FString MaxAllowedVersion;
+};
 
 
 /**
@@ -20,11 +41,26 @@ class ITurnkeySupportModule
 public:
 
 	/**
-	 *
-	 * @return	The newly-created menu widget
+	 * Runs Turnkey to get the Sdk information for all known platforms
 	 */
 	virtual TSharedRef<SWidget> MakeTurnkeyMenu() const = 0;
 
+	/**
+	 * @return	The newly-created menu widget
+	 */
+	virtual void UpdateSdkInfo() = 0;
+
+	/**
+	 * Runs Turnkey to get the Sdk information for a list of devices
+	 */
+	virtual void UpdateSdkInfoForDevices(TArray<FString> DeviceIds) = 0;
+
+
+	virtual FTurnkeySdkInfo GetSdkInfo(FName PlatformName, bool bBlockIfQuerying = true) const = 0;
+	virtual FTurnkeySdkInfo GetSdkInfoForDeviceId(const FString& DeviceId) const = 0;
+	// @todo turnkey look into remove this
+	virtual void ClearDeviceStatus(FName PlatformName=NAME_None) = 0;
+
 public:
 
 	/**
diff --git a/Engine/Source/Editor/TurnkeySupport/Public/TurnkeySupport.h b/Engine/Source/Editor/TurnkeySupport/Public/TurnkeySupport.h
--- a/Engine/Source/Editor/TurnkeySupport/Public/TurnkeySupport.h
+++ b/Engine/Source/Editor/TurnkeySupport/Public/TurnkeySupport.h
@@ -3,21 +3,28 @@
 #pragma once
 
 
-/* Boilerplate
- *****************************************************************************/
-
-#include "Misc/MonolithicHeaderBoilerplate.h"
-MONOLITHIC_HEADER_BOILERPLATE()
-
 
 /* Public Dependencies
  *****************************************************************************/
 
 #include "Modules/ModuleInterface.h"
-#include "SlateBasics.h"
+#include "SlateFwd.h"
+#include "Widgets/SWidget.h"
+#include "Framework/Commands/UICommandList.h"
+#include "Framework/MultiBox/MultiBoxExtender.h"
 
 
 /* Public Includes
  *****************************************************************************/
 
 #include "Interfaces/ITurnkeySupportModule.h"
+
+
+
+DECLARE_LOG_CATEGORY_EXTERN(LogTurnkeySupport, Log, All);
+
+FString ConvertToDDPIPlatform(const FString& Platform);
+FName ConvertToDDPIPlatform(const FName& Platform);
+FString ConvertToUATPlatform(const FString& Platform);
+FString ConvertToUATDeviceId(const FString& DeviceId);
+FString ConvertToDDPIDeviceId(const FString& DeviceId);
diff --git a/Engine/Source/Editor/UnrealEd/Private/Kismet2/DebuggerCommands.cpp b/Engine/Source/Editor/UnrealEd/Private/Kismet2/DebuggerCommands.cpp
--- a/Engine/Source/Editor/UnrealEd/Private/Kismet2/DebuggerCommands.cpp
+++ b/Engine/Source/Editor/UnrealEd/Private/Kismet2/DebuggerCommands.cpp
@@ -2146,15 +2146,15 @@ bool FInternalPlayWorldCommandCallbacks::IsReadyToLaunchOnDevice(FString DeviceI
 {
 	int32 Index = 0;
 	DeviceId.FindChar(TEXT('@'), Index);
-	FString PlatformName = DeviceId.Left(Index);
+	FName PlatformName = FName(*DeviceId.Left(Index));
 
-	const PlatformInfo::FTargetPlatformInfo* const PlatformInfo = PlatformInfo::FindPlatformInfo(FName(*PlatformName));
-	checkf(PlatformInfo, TEXT("Unable to find PlatformInfo for %s"), *PlatformName);
+	const PlatformInfo::FTargetPlatformInfo* PlatformInfo = PlatformInfo::FindPlatformInfo(PlatformName);
+	checkf(PlatformInfo, TEXT("Unable to find PlatformInfo for %s"), *PlatformName.ToString());
 
 	FGameProjectGenerationModule& GameProjectModule = FModuleManager::LoadModuleChecked<FGameProjectGenerationModule>(TEXT("GameProjectGeneration"));
 	bool bHasCode = GameProjectModule.Get().ProjectHasCodeFiles();
 
-	if (PlatformInfo->DataDrivenPlatformInfo->GetSdkStatus() != DDPIPlatformSdkStatus::Valid)
+	if (ITurnkeySupportModule::Get().GetSdkInfo(PlatformName).Status!= ETurnkeyPlatformSdkStatus::Valid)
 	{
 		IMainFrameModule& MainFrameModule = FModuleManager::GetModuleChecked<IMainFrameModule>(TEXT("MainFrame"));
 		MainFrameModule.BroadcastMainFrameSDKNotInstalled(PlatformInfo->Name.ToString(), PlatformInfo->DataDrivenPlatformInfo->SDKTutorial);
@@ -2176,7 +2176,7 @@ bool FInternalPlayWorldCommandCallbacks::IsReadyToLaunchOnDevice(FString DeviceI
 		int32 Result = Platform->CheckRequirements(bHasCode, BuildConfiguration, bEnableAssetNativization, NotInstalledTutorialLink, DocumentationLink, CustomizedLogMessage);
 
 		// report to analytics
-		FEditorAnalytics::ReportBuildRequirementsFailure(TEXT("Editor.LaunchOn.Failed"), PlatformName, bHasCode, Result);
+		FEditorAnalytics::ReportBuildRequirementsFailure(TEXT("Editor.LaunchOn.Failed"), PlatformName.ToString(), bHasCode, Result);
 
 		// report to message log
 		bool UnrecoverableError = false;
diff --git a/Engine/Source/Runtime/Core/Private/Misc/DataDrivenPlatformInfoRegistry.cpp b/Engine/Source/Runtime/Core/Private/Misc/DataDrivenPlatformInfoRegistry.cpp
--- a/Engine/Source/Runtime/Core/Private/Misc/DataDrivenPlatformInfoRegistry.cpp
+++ b/Engine/Source/Runtime/Core/Private/Misc/DataDrivenPlatformInfoRegistry.cpp
@@ -7,15 +7,12 @@
 #include "Misc/ConfigCacheIni.h"
 #include "Modules/ModuleManager.h"
 
-#include "Async/Async.h"
-#include "Misc/MonitoredProcess.h"
 
 namespace 
 {
 	TMap<FName, FDataDrivenPlatformInfo> DataDrivenPlatforms;
 	TArray<FName> SortedPlatformNames;
 	TArray<const FDataDrivenPlatformInfo*> SortedPlatformInfos;
-	FCriticalSection DDPILocker;
 }
 
 #if DDPI_HAS_EXTENDED_PLATFORMINFO_DATA
@@ -397,383 +394,6 @@ const TArray<const FDataDrivenPlatformInfo*>& FDataDrivenPlatformInfoRegistry::G
 }
 
 
-#if DDPI_HAS_EXTENDED_PLATFORMINFO_DATA
-
-// some shared functionality
-static void PrepForTurnkeyReport(FString& Command, FString& BaseCommandline, FString& ReportFilename)
-{
- 	static int ReportIndex = 0;
-
-	FString LogFilename = FPaths::ConvertRelativePathToFull(FPaths::Combine(FPaths::ProjectIntermediateDir(), *FString::Printf(TEXT("TurnkeyLog_%d.log"), ReportIndex)));
-	ReportFilename = FPaths::ConvertRelativePathToFull(FPaths::Combine(FPaths::ProjectIntermediateDir(), *FString::Printf(TEXT("TurnkeyReport_%d.log"), ReportIndex++)));
-
-	// make sure intermediate directory exists
-	IFileManager::Get().MakeDirectory(*FPaths::ProjectIntermediateDir());
-
-	Command = TEXT("{EngineDir}Build/BatchFiles/RunuAT");
-//	Command = TEXT("{EngineDir}/Binaries/DotNET/AutomationTool.exe");
-	BaseCommandline = FString::Printf(TEXT("Turnkey -utf8output -WaitForUATMutex -command=VerifySdk -ReportFilename=\"%s\" -log=\"%s\""), *ReportFilename, *LogFilename);
-
-	// convert into appropriate calls for the current platform
-	FPlatformProcess::ModifyCreateProcParams(Command, BaseCommandline, FGenericPlatformProcess::ECreateProcHelperFlags::AppendPlatformScriptExtension | FGenericPlatformProcess::ECreateProcHelperFlags::RunThroughShell);
-}
-
-static FString ConvertToDDPIPlatform(const FString& Platform)
-{
-	FString  New = Platform.Replace(TEXT("Editor"), TEXT("")).Replace(TEXT("Client"), TEXT("")).Replace(TEXT("Server"), TEXT(""));
-	if (New == TEXT("Win64"))
-	{
-		New = TEXT("Windows");
-	}
-	return New;
-}
-
-static FString ConvertToUATPlatform(const FString& Platform)
-{
-	FString New = ConvertToDDPIPlatform(Platform);
-	if (New == TEXT("Windows"))
-	{
-		New = TEXT("Win64");
-	}
-	return New;
-}
-
-static FString ConvertToUATDeviceId(const FString& DeviceId)
-{
-	TArray<FString> PlatformAndDevice;
-	DeviceId.ParseIntoArray(PlatformAndDevice, TEXT("@"), true);
-
-	return FString::Printf(TEXT("%s@%s"), *ConvertToUATPlatform(PlatformAndDevice[0]), *PlatformAndDevice[1]);
-}
-
-static FString ConvertToDDPIDeviceId(const FString& DeviceId)
-{
-	TArray<FString> PlatformAndDevice;
-	DeviceId.ParseIntoArray(PlatformAndDevice, TEXT("@"), true);
-
-	return FString::Printf(TEXT("%s@%s"), *ConvertToDDPIPlatform(PlatformAndDevice[0]), *PlatformAndDevice[1]);
-}
-
-bool GetSdkInfoFromTurnkey(FString Line, FString& PlatformName, FString& DeviceId, FDDPISdkInfo& SdkInfo)
-{
-	int32 Colon = Line.Find(TEXT(": "));
-
-	if (Colon < 0)
-	{
-		return false;
-	}
-
-	// break up the string
-	PlatformName = Line.Mid(0, Colon);
-	FString Info = Line.Mid(Colon + 2);
-
-	int32 AtSign = PlatformName.Find(TEXT("@"));
-	if (AtSign > 0)
-	{
-		// return the platform@name as the deviceId, then remove the @name part for the platform
-		DeviceId = ConvertToDDPIDeviceId(PlatformName);
-		PlatformName = PlatformName.Mid(0, AtSign);
-	}
-
-	// get the DDPI name
-	PlatformName = ConvertToDDPIPlatform(PlatformName);
-
-	// parse out the results from the (key=val, key=val) result from turnkey
-	FString StatusString;
-	FString FlagsString;
-	FParse::Value(*Info, TEXT("Status="), StatusString);
-	FParse::Value(*Info, TEXT("Flags="), FlagsString);
-	FParse::Value(*Info, TEXT("Installed="), SdkInfo.InstalledVersion);
-	FParse::Value(*Info, TEXT("AutoSDK="), SdkInfo.AutoSDKVersion);
-	FParse::Value(*Info, TEXT("MinAllowed="), SdkInfo.MinAllowedVersion);
-	FParse::Value(*Info, TEXT("MaxAllowed="), SdkInfo.MaxAllowedVersion);
-
-	SdkInfo.Status = DDPIPlatformSdkStatus::Unknown;
-	if (StatusString == TEXT("Valid"))
-	{
-		SdkInfo.Status = DDPIPlatformSdkStatus::Valid;
-	}
-	else
-	{
-		if (FlagsString.Contains(TEXT("AutoSdk_InvalidVersionExists")) || FlagsString.Contains(TEXT("InstalledSdk_InvalidVersionExists")))
-		{
-			SdkInfo.Status = DDPIPlatformSdkStatus::OutOfDate;
-		}
-		else
-		{
-			SdkInfo.Status = DDPIPlatformSdkStatus::NoSdk;
-		}
-	}
-
-	return true;
-}
-
-void FDataDrivenPlatformInfoRegistry::UpdateSdkStatus()
-{
-	// make sure we've read in the inis
-	GetAllPlatformInfos();
-
-	// don't run UAT from commandlets (like the cooker) that are often launched from UAT and this will go poorly
-	if (IsRunningCommandlet())
-	{
-		for (auto& It : DataDrivenPlatforms)
-		{
-			It.Value.SdkInfo.Status = DDPIPlatformSdkStatus::Unknown;
-
-			// reset the per-device status when querying general Sdk status
-			It.Value.ClearDeviceStatus();
-		}
-
-		return;
-	}
-
-
-	FString Command, BaseCommandline, ReportFilename;
-	PrepForTurnkeyReport(Command, BaseCommandline, ReportFilename);
-	FString Commandline = BaseCommandline + FString(TEXT(" -platform=")) + FString::JoinBy(DataDrivenPlatforms, TEXT("+"), [](TPair<FName, FDataDrivenPlatformInfo> Pair) { return ConvertToUATPlatform(Pair.Key.ToString()); });
-
-	UE_LOG(LogInit, Log, TEXT("Running Turnkey SDK detection: '%s %s'"), *Command, *Commandline);
-
-	{
-		FScopeLock Lock(&DDPILocker);
-
-		// reset status to unknown
-		for (auto& It : DataDrivenPlatforms)
-		{
-			It.Value.SdkInfo.Status = DDPIPlatformSdkStatus::Querying;
-
-			// reset the per-device status when querying general Sdk status
-			It.Value.ClearDeviceStatus();
-		}
-	}
-
-	FMonitoredProcess* TurnkeyProcess = new FMonitoredProcess(Command, Commandline, true, false);
-	TurnkeyProcess->OnCompleted().BindLambda([ReportFilename, TurnkeyProcess](int32 ExitCode)
-	{
-		AsyncTask(ENamedThreads::GameThread, [ReportFilename, TurnkeyProcess, ExitCode]()
-		{
-			FScopeLock Lock(&DDPILocker);
-
-			if (ExitCode == 0 || ExitCode == 10)
-			{
-				TArray<FString> Contents;
-				if (FFileHelper::LoadFileToStringArray(Contents, *ReportFilename))
-				{
-					for (FString& Line : Contents)
-					{
-						UE_LOG(LogTemp, Log, TEXT("Turnkey Platform: %s"), *Line);
-
-						// parse a Turnkey line
-						FString PlatformName, Unused;
-						FDDPISdkInfo SdkInfo;
-						if (GetSdkInfoFromTurnkey(Line, PlatformName, Unused, SdkInfo) == false)
-						{
-							continue;
-						}
-
-						// check if we had already set a ManualSDK - and don't set it again. Because of the way AutoSDKs are activated in the editor after the first call to Turnkey,
-						// future calls to Turnkey will inherit the AutoSDK env vars, and it won't be able to determine the manual SDK versions anymore. If we use the editor to
-						// install an SDK via Turnkey, it will directly update the installed version based on the result of that command, not this Update operation
-
-						FString OriginalManualInstallValue = DataDrivenPlatforms[*PlatformName].SdkInfo.InstalledVersion;
-
-						// set it into the platform
-						DataDrivenPlatforms[*PlatformName].SdkInfo = SdkInfo;
-
-						// restore the original installed version if it set after the first time
-						if (OriginalManualInstallValue.Len() > 0)
-						{
-							DataDrivenPlatforms[*PlatformName].SdkInfo.InstalledVersion = OriginalManualInstallValue;
-						}
-
-
-						UE_LOG(LogTemp, Log, TEXT("[TEST] Turnkey Platform: %s - %d, Installed: %s, AudoSDK: %s, Allowed: %s-%s"), *PlatformName, (int)SdkInfo.Status, *SdkInfo.InstalledVersion,
-							*SdkInfo.AutoSDKVersion, *SdkInfo.MinAllowedVersion, *SdkInfo.MaxAllowedVersion);
-					}
-				}
-			}
-			else
-			{
-				for (auto& It : DataDrivenPlatforms)
-				{
-					It.Value.SdkInfo.Status = DDPIPlatformSdkStatus::Error;
-					It.Value.SdkInfo.SdkErrorInformation = FText::Format(NSLOCTEXT("Turnkey", "TurnkeyError_ReturnedError", "Turnkey returned an error, code {0}"), { ExitCode });
-
-					// @todo turnkey error description!
-				}
-			}
-
-
-			for (auto& It : DataDrivenPlatforms)
-			{
-				if (It.Value.SdkInfo.Status == DDPIPlatformSdkStatus::Querying)
-				{
-					if (It.Value.bIsFakePlatform)
-					{
-						It.Value.SdkInfo.Status = DDPIPlatformSdkStatus::Unknown;
-					}
-					else
-					{
-						It.Value.SdkInfo.Status = DDPIPlatformSdkStatus::Error;
-						It.Value.SdkInfo.SdkErrorInformation = NSLOCTEXT("Turnkey", "TurnkeyError_NotReturned", "The platform's Sdk status was not returned from Turnkey");
-					}
-				}
-			}
-
-			// cleanup
-			delete TurnkeyProcess;
-			IFileManager::Get().Delete(*ReportFilename);
-		});
-	});
-
-	// run it
-	TurnkeyProcess->Launch();
-}
-
-FDataDrivenPlatformInfo& FDataDrivenPlatformInfoRegistry::DeviceIdToInfo(FString DeviceId, FString* OutDeviceName)
-{
-	TArray<FString> PlatformAndDevice;
-	DeviceId.ParseIntoArray(PlatformAndDevice, TEXT("@"), true);
-
-	if (OutDeviceName)
-	{
-		*OutDeviceName = PlatformAndDevice[1];
-	}
-
-	FString DDPIPlatformName = ConvertToDDPIPlatform(PlatformAndDevice[0]);
-
-	checkf(DataDrivenPlatforms.Contains(*DDPIPlatformName), TEXT("DataDrivenPlatforms map did not contain the DDPI Platform %s"), *DDPIPlatformName);
-
-	// have to convert back to Windows from Win64
-	return DataDrivenPlatforms[*DDPIPlatformName];
-
-}
-
-void FDataDrivenPlatformInfoRegistry::UpdateDeviceSdkStatus(TArray<FString> PlatformDeviceIds)
-{
-	FString Command, BaseCommandline, ReportFilename;
-	PrepForTurnkeyReport(Command, BaseCommandline, ReportFilename);
-
-	// the platform part of the Id may need to be converted to be turnkey (ie UBT) proper
-
-	FString Commandline = BaseCommandline + FString(TEXT(" -Device=")) + FString::JoinBy(PlatformDeviceIds, TEXT("+"), [](FString Id) { return ConvertToUATDeviceId(Id); });
-
-	UE_LOG(LogInit, Log, TEXT("Running Turnkey SDK detection: '%s %s'"), *Command, *Commandline);
-
-	{
-		FScopeLock Lock(&DDPILocker);
-
-		// set status to querying
-		FDDPISdkInfo DefaultInfo;
-		DefaultInfo.Status = DDPIPlatformSdkStatus::Querying;
-		for (const FString& Id : PlatformDeviceIds)
-		{
-			DeviceIdToInfo(Id).PerDeviceStatus.Add(ConvertToDDPIDeviceId(Id), DefaultInfo);
-		}
-	}
-
-	FMonitoredProcess* TurnkeyProcess = new FMonitoredProcess(Command, Commandline, true, false);
-	TurnkeyProcess->OnCompleted().BindLambda([ReportFilename, TurnkeyProcess, PlatformDeviceIds](int32 ExitCode)
-	{
-		AsyncTask(ENamedThreads::GameThread, [ReportFilename, TurnkeyProcess, PlatformDeviceIds, ExitCode]()
-		{
-			FScopeLock Lock(&DDPILocker);
-
-			if (ExitCode == 0 || ExitCode == 10)
-			{
-				TArray<FString> Contents;
-				if (FFileHelper::LoadFileToStringArray(Contents, *ReportFilename))
-				{
-					for (FString& Line : Contents)
-					{
-						FString PlatformName, DDPIDeviceId;
-						FDDPISdkInfo SdkInfo;
-						if (GetSdkInfoFromTurnkey(Line, PlatformName, DDPIDeviceId, SdkInfo) == false)
-						{
-							continue;
-						}
-
-						// skip over non-device lines
-						if (DDPIDeviceId.Len() == 0)
-						{
-							continue;
-						}
-
-						UE_LOG(LogTemp, Log, TEXT("Turnkey Device: %s"), *Line);
-
-						DeviceIdToInfo(DDPIDeviceId).PerDeviceStatus[DDPIDeviceId] = SdkInfo;
-
-						UE_LOG(LogTemp, Log, TEXT("[TEST] Turnkey Device: %s - %d, Installed: %s, Allowed: %s-%s"), *DDPIDeviceId, (int)SdkInfo.Status, *SdkInfo.InstalledVersion,
-							*SdkInfo.MinAllowedVersion, *SdkInfo.MaxAllowedVersion);
-					}
-				}
-			}
-
-			for (const FString& Id : PlatformDeviceIds)
-			{
-				FDataDrivenPlatformInfo& Info = DeviceIdToInfo(Id);
-				
-				FDDPISdkInfo& SdkInfo = Info.PerDeviceStatus[ConvertToDDPIDeviceId(Id)];
-				if (SdkInfo.Status == DDPIPlatformSdkStatus::Querying)
-				{
-					SdkInfo.Status = DDPIPlatformSdkStatus::Error;
-					SdkInfo.SdkErrorInformation = NSLOCTEXT("Turnkey", "TurnkeyError_DeviceNotReturned", "A device's Sdk status was not returned from Turnkey");
-				}
-			}
-
-			// cleanup
-			delete TurnkeyProcess;
-			IFileManager::Get().Delete(*ReportFilename);
-		});
-	});
-
-	// run it
-	TurnkeyProcess->Launch();
-}
-
-void FDataDrivenPlatformInfoRegistry::ClearDeviceStatus(FName PlatformName)
-{
-	FScopeLock Lock(&DDPILocker);
-
-	if (PlatformName != NAME_None)
-	{
-		for (auto It : DataDrivenPlatforms)
-		{
-			It.Value.ClearDeviceStatus();
-		}
-	}
-	else
-	{
-		if (DataDrivenPlatforms.Contains(PlatformName))
-		{
-			DataDrivenPlatforms[PlatformName].ClearDeviceStatus();
-		}
-	}
-}
-
-DDPIPlatformSdkStatus FDataDrivenPlatformInfo::GetStatusForDeviceId(const FString& DeviceId) const
-{
-	// return the status, or Unknown if not known
-	return PerDeviceStatus.FindRef(ConvertToDDPIDeviceId(DeviceId)).Status;
-}
-
-const FDDPISdkInfo& FDataDrivenPlatformInfo::GetSdkInfoForDeviceId(const FString& DeviceId) const
-{
-	return *PerDeviceStatus.Find(ConvertToDDPIDeviceId(DeviceId));
-}
-
-
-void FDataDrivenPlatformInfo::ClearDeviceStatus()
-{
-	for (auto& Pair : PerDeviceStatus)
-	{
-		Pair.Value.Status = DDPIPlatformSdkStatus::Unknown;
-	}
-}
-
-#endif
-
-
 const TArray<FString>& FDataDrivenPlatformInfoRegistry::GetValidPlatformDirectoryNames()
 {
 	static bool bHasSearchedForPlatforms = false;
diff --git a/Engine/Source/Runtime/Core/Public/Misc/DataDrivenPlatformInfoRegistry.h b/Engine/Source/Runtime/Core/Public/Misc/DataDrivenPlatformInfoRegistry.h
--- a/Engine/Source/Runtime/Core/Public/Misc/DataDrivenPlatformInfoRegistry.h
+++ b/Engine/Source/Runtime/Core/Public/Misc/DataDrivenPlatformInfoRegistry.h
@@ -13,28 +13,6 @@
 
 #if DDPI_HAS_EXTENDED_PLATFORMINFO_DATA
 
-enum class DDPIPlatformSdkStatus : uint8
-{
-	Unknown,
-	Querying,
-	Valid,
-	OutOfDate,
-	NoSdk,
-	Error,
-	FlashValid,
-	FlashOutOfDate,
-	// @todo turnkey: add AutoSdkValid and ManualSdkValid, with Valid a Combination of both
-};
-
-struct FDDPISdkInfo
-{
-	DDPIPlatformSdkStatus Status;
-	FText SdkErrorInformation;
-	FString InstalledVersion;
-	FString AutoSDKVersion; // only valid for platform, not device
-	FString MinAllowedVersion;
-	FString MaxAllowedVersion;
-};
 
 /** Available icon sizes (see FPlatformIconPaths) */
 enum class EPlatformIconSize : uint8
@@ -154,26 +132,7 @@ struct FDataDrivenPlatformInfo
 
 #if DDPI_HAS_EXTENDED_PLATFORMINFO_DATA
 
-private:
-	// Information about the validity of using a platform, discovered via Turnkey
-	FDDPISdkInfo SdkInfo;
-
-	// Information about the validity of each connected device (by string, discovered by Turnkey)
-	TMap<FString, FDDPISdkInfo> PerDeviceStatus;
-
 public:
-	DDPIPlatformSdkStatus GetSdkStatus(bool bBlockIfQuerying = true) const
-	{
-		check(!(bBlockIfQuerying && SdkInfo.Status == DDPIPlatformSdkStatus::Querying));
-		return SdkInfo.Status;
-	}
-
-	const FDDPISdkInfo& GetSdkInfo(bool bBlockIfQuerying = true) const
-	{
-		check(!(bBlockIfQuerying && SdkInfo.Status == DDPIPlatformSdkStatus::Querying));
-		return SdkInfo;
-	}
-
 
 
 	// setting moved from PlatformInfo::FTargetPlatformInfo
@@ -217,11 +176,6 @@ public:
 	bool bHasCompiledTargetSupport;
 
 
-	// Get the status of a device, or Unknown if not specified
-	CORE_API DDPIPlatformSdkStatus GetStatusForDeviceId(const FString& DeviceId) const;
-	CORE_API const FDDPISdkInfo& GetSdkInfoForDeviceId(const FString& DeviceId) const;
-	CORE_API void ClearDeviceStatus();
-
 
 	/** Get the icon name (for FEditorStyle) used by the given icon type for this platform */
 	FName GetIconStyleName(const EPlatformIconSize InIconSize) const
@@ -287,6 +241,7 @@ struct CORE_API FDataDrivenPlatformInfoRegistry
 	static const FDataDrivenPlatformInfo& GetPlatformInfo(FName PlatformName);
 	static const FDataDrivenPlatformInfo& GetPlatformInfo(const char* PlatformName);
 
+
 	// get just names or just infos, 
 	static const TArray<FName> GetSortedPlatformNames();
 	static const TArray<const FDataDrivenPlatformInfo*>& GetSortedPlatformInfos();
@@ -325,22 +280,20 @@ struct CORE_API FDataDrivenPlatformInfoRegistry
 	static bool HasCompiledSupportForPlatform(FName PlatformName, EPlatformNameType PlatformNameType);
 
 	/**
-	 * Runs UAT to update Sdk status (can be called after user chooses an option that had a bad Sdk status)
+	 * Wipes out cached device status for all devices in a platform (or all platforms if PlatformName is empty)
 	 */
-	static void UpdateSdkStatus();
+	static void ClearDeviceStatus(FName PlatformName);
 
-	/**
-	 * Run UAT to update per device flash/software information (each Id should be in the standard DeviceId format of Platform@Device)
-	 */
-	static void UpdateDeviceSdkStatus(TArray<FString> PlatformDeviceIds);
+	static FDataDrivenPlatformInfo& DeviceIdToInfo(FString DeviceId, FString* OutDeviceName = nullptr);
 
+private:
 	/**
-	 * Wipes out cached device status for all devices in a platform (or all platforms if PlatformName is empty)
+	 * Get a modifiable DDPI object, for Turnkey to update it's info
 	 */
-	static void ClearDeviceStatus(FName PlatformName);
+	friend class FTurkeySupportModule;
+	static FDataDrivenPlatformInfo& GetMutablePlatformInfo(FName PlatformName);
+	static TMap<FName, FDataDrivenPlatformInfo>& GetMutablePlatformInfos();
 
-private:
-	static FDataDrivenPlatformInfo& DeviceIdToInfo(FString DeviceId, FString* OutDeviceName = nullptr);
 
 #endif
 };
