Adding a new file I missed in last checkin

diff --git a/Engine/Source/Editor/TurnkeySupport/Private/TurnkeyWrapper.cpp b/Engine/Source/Editor/TurnkeySupport/Private/TurnkeyWrapper.cpp
new file mode 100644
--- /dev/null
+++ b/Engine/Source/Editor/TurnkeySupport/Private/TurnkeyWrapper.cpp
@@ -0,0 +1,382 @@
+// Copyright Epic Games, Inc. All Rights Reserved.
+
+#include "TurnkeySupportModule.h"
+#include "Async/Async.h"
+#include "Misc/MonitoredProcess.h"
+
+FString ConvertToDDPIPlatform(const FString& Platform)
+{
+	FString New = Platform.Replace(TEXT("Editor"), TEXT("")).Replace(TEXT("Client"), TEXT("")).Replace(TEXT("Server"), TEXT(""));
+	if (New == TEXT("Win64"))
+	{
+		New = TEXT("Windows");
+	}
+	return New;
+}
+
+FName ConvertToDDPIPlatform(const FName& Platform)
+{
+	return FName(*ConvertToDDPIPlatform(Platform.ToString()));
+}
+
+FString ConvertToUATPlatform(const FString& Platform)
+{
+	FString New = ConvertToDDPIPlatform(Platform);
+	if (New == TEXT("Windows"))
+	{
+		New = TEXT("Win64");
+	}
+	return New;
+}
+
+FString ConvertToUATDeviceId(const FString& DeviceId)
+{
+	TArray<FString> PlatformAndDevice;
+	DeviceId.ParseIntoArray(PlatformAndDevice, TEXT("@"), true);
+
+	return FString::Printf(TEXT("%s@%s"), *ConvertToUATPlatform(PlatformAndDevice[0]), *PlatformAndDevice[1]);
+}
+
+FString ConvertToDDPIDeviceId(const FString& DeviceId)
+{
+	TArray<FString> PlatformAndDevice;
+	DeviceId.ParseIntoArray(PlatformAndDevice, TEXT("@"), true);
+
+	return FString::Printf(TEXT("%s@%s"), *ConvertToDDPIPlatform(PlatformAndDevice[0]), *PlatformAndDevice[1]);
+}
+
+
+// 
+// // some shared functionality
+// static void PrepForTurnkeyReport(FString& Command, FString& BaseCommandline, FString& ReportFilename)
+// {
+// 	static int ReportIndex = 0;
+// 
+// 	FString LogFilename = FPaths::ConvertRelativePathToFull(FPaths::Combine(FPaths::ProjectIntermediateDir(), *FString::Printf(TEXT("TurnkeyLog_%d.log"), ReportIndex)));
+// 	ReportFilename = FPaths::ConvertRelativePathToFull(FPaths::Combine(FPaths::ProjectIntermediateDir(), *FString::Printf(TEXT("TurnkeyReport_%d.log"), ReportIndex++)));
+// 
+// 	// make sure intermediate directory exists
+// 	IFileManager::Get().MakeDirectory(*FPaths::ProjectIntermediateDir());
+// 
+// 	Command = TEXT("{EngineDir}Build/BatchFiles/RunuAT");
+// 	//	Command = TEXT("{EngineDir}/Binaries/DotNET/AutomationTool.exe");
+// 	BaseCommandline = FString::Printf(TEXT("Turnkey -utf8output -WaitForUATMutex -command=VerifySdk -ReportFilename=\"%s\" -log=\"%s\""), *ReportFilename, *LogFilename);
+// 
+// 	// convert into appropriate calls for the current platform
+// 	FPlatformProcess::ModifyCreateProcParams(Command, BaseCommandline, FGenericPlatformProcess::ECreateProcHelperFlags::AppendPlatformScriptExtension | FGenericPlatformProcess::ECreateProcHelperFlags::RunThroughShell);
+// }
+// 
+// static FString ConvertToDDPIPlatform(const FString& Platform)
+// {
+// 	FString  New = Platform.Replace(TEXT("Editor"), TEXT("")).Replace(TEXT("Client"), TEXT("")).Replace(TEXT("Server"), TEXT(""));
+// 	if (New == TEXT("Win64"))
+// 	{
+// 		New = TEXT("Windows");
+// 	}
+// 	return New;
+// }
+// 
+// static FString ConvertToUATPlatform(const FString& Platform)
+// {
+// 	FString New = ConvertToDDPIPlatform(Platform);
+// 	if (New == TEXT("Windows"))
+// 	{
+// 		New = TEXT("Win64");
+// 	}
+// 	return New;
+// }
+// 
+// static FString ConvertToUATDeviceId(const FString& DeviceId)
+// {
+// 	TArray<FString> PlatformAndDevice;
+// 	DeviceId.ParseIntoArray(PlatformAndDevice, TEXT("@"), true);
+// 
+// 	return FString::Printf(TEXT("%s@%s"), *ConvertToUATPlatform(PlatformAndDevice[0]), *PlatformAndDevice[1]);
+// }
+// 
+// static FString ConvertToDDPIDeviceId(const FString& DeviceId)
+// {
+// 	TArray<FString> PlatformAndDevice;
+// 	DeviceId.ParseIntoArray(PlatformAndDevice, TEXT("@"), true);
+// 
+// 	return FString::Printf(TEXT("%s@%s"), *ConvertToDDPIPlatform(PlatformAndDevice[0]), *PlatformAndDevice[1]);
+// }
+// 
+// bool GetSdkInfoFromTurnkey(FString Line, FString& PlatformName, FString& DeviceId, FDDPISdkInfo& SdkInfo)
+// {
+// 	int32 Colon = Line.Find(TEXT(": "));
+// 
+// 	if (Colon < 0)
+// 	{
+// 		return false;
+// 	}
+// 
+// 	// break up the string
+// 	PlatformName = Line.Mid(0, Colon);
+// 	FString Info = Line.Mid(Colon + 2);
+// 
+// 	int32 AtSign = PlatformName.Find(TEXT("@"));
+// 	if (AtSign > 0)
+// 	{
+// 		// return the platform@name as the deviceId, then remove the @name part for the platform
+// 		DeviceId = ConvertToDDPIDeviceId(PlatformName);
+// 		PlatformName = PlatformName.Mid(0, AtSign);
+// 	}
+// 
+// 	// get the DDPI name
+// 	PlatformName = ConvertToDDPIPlatform(PlatformName);
+// 
+// 	// parse out the results from the (key=val, key=val) result from turnkey
+// 	FString StatusString;
+// 	FString FlagsString;
+// 	FParse::Value(*Info, TEXT("Status="), StatusString);
+// 	FParse::Value(*Info, TEXT("Flags="), FlagsString);
+// 	FParse::Value(*Info, TEXT("Installed="), SdkInfo.InstalledVersion);
+// 	FParse::Value(*Info, TEXT("AutoSDK="), SdkInfo.AutoSDKVersion);
+// 	FParse::Value(*Info, TEXT("MinAllowed="), SdkInfo.MinAllowedVersion);
+// 	FParse::Value(*Info, TEXT("MaxAllowed="), SdkInfo.MaxAllowedVersion);
+// 
+// 	SdkInfo.Status = DDPIPlatformSdkStatus::Unknown;
+// 	if (StatusString == TEXT("Valid"))
+// 	{
+// 		SdkInfo.Status = DDPIPlatformSdkStatus::Valid;
+// 	}
+// 	else
+// 	{
+// 		if (FlagsString.Contains(TEXT("AutoSdk_InvalidVersionExists")) || FlagsString.Contains(TEXT("InstalledSdk_InvalidVersionExists")))
+// 		{
+// 			SdkInfo.Status = DDPIPlatformSdkStatus::OutOfDate;
+// 		}
+// 		else
+// 		{
+// 			SdkInfo.Status = DDPIPlatformSdkStatus::NoSdk;
+// 		}
+// 	}
+// 
+// 	return true;
+// }
+// 
+// FDataDrivenPlatformInfo& FDataDrivenPlatformInfoRegistry::DeviceIdToInfo(FString DeviceId, FString* OutDeviceName)
+// {
+// 	TArray<FString> PlatformAndDevice;
+// 	DeviceId.ParseIntoArray(PlatformAndDevice, TEXT("@"), true);
+// 
+// 	if (OutDeviceName)
+// 	{
+// 		*OutDeviceName = PlatformAndDevice[1];
+// 	}
+// 
+// 	FString DDPIPlatformName = ConvertToDDPIPlatform(PlatformAndDevice[0]);
+// 
+// 	checkf(DataDrivenPlatforms.Contains(*DDPIPlatformName), TEXT("DataDrivenPlatforms map did not contain the DDPI Platform %s"), *DDPIPlatformName);
+// 
+// 	// have to convert back to Windows from Win64
+// 	return DataDrivenPlatforms[*DDPIPlatformName];
+// 
+// }
+// 
+// 
+// void FTurnkeySupportModule::UpdateSdkInfo() const
+// {
+// 	// make sure we've read in the inis
+// 	GetAllPlatformInfos();
+// 
+// 	// don't run UAT from commandlets (like the cooker) that are often launched from UAT and this will go poorly
+// 	if (IsRunningCommandlet())
+// 	{
+// 		for (auto& It : DataDrivenPlatforms)
+// 		{
+// 			It.Value.SdkInfo.Status = DDPIPlatformSdkStatus::Unknown;
+// 
+// 			// reset the per-device status when querying general Sdk status
+// 			It.Value.ClearDeviceStatus();
+// 		}
+// 
+// 		return;
+// 	}
+// 
+// 
+// 	FString Command, BaseCommandline, ReportFilename;
+// 	PrepForTurnkeyReport(Command, BaseCommandline, ReportFilename);
+// 	FString Commandline = BaseCommandline + FString(TEXT(" -platform=")) + FString::JoinBy(DataDrivenPlatforms, TEXT("+"), [](TPair<FName, FDataDrivenPlatformInfo> Pair) { return ConvertToUATPlatform(Pair.Key.ToString()); });
+// 
+// 	UE_LOG(LogInit, Log, TEXT("Running Turnkey SDK detection: '%s %s'"), *Command, *Commandline);
+// 
+// 	{
+// 		FScopeLock Lock(&DDPILocker);
+// 
+// 		// reset status to unknown
+// 		for (auto& It : DataDrivenPlatforms)
+// 		{
+// 			It.Value.SdkInfo.Status = DDPIPlatformSdkStatus::Querying;
+// 
+// 			// reset the per-device status when querying general Sdk status
+// 			It.Value.ClearDeviceStatus();
+// 		}
+// 	}
+// 
+// 	FMonitoredProcess* TurnkeyProcess = new FMonitoredProcess(Command, Commandline, true, false);
+// 	TurnkeyProcess->OnCompleted().BindLambda([ReportFilename, TurnkeyProcess](int32 ExitCode)
+// 	{
+// 		AsyncTask(ENamedThreads::GameThread, [ReportFilename, TurnkeyProcess, ExitCode]()
+// 		{
+// 			FScopeLock Lock(&DDPILocker);
+// 
+// 			if (ExitCode == 0 || ExitCode == 10)
+// 			{
+// 				TArray<FString> Contents;
+// 				if (FFileHelper::LoadFileToStringArray(Contents, *ReportFilename))
+// 				{
+// 					for (FString& Line : Contents)
+// 					{
+// 						UE_LOG(LogTemp, Log, TEXT("Turnkey Platform: %s"), *Line);
+// 
+// 						// parse a Turnkey line
+// 						FString PlatformName, Unused;
+// 						FDDPISdkInfo SdkInfo;
+// 						if (GetSdkInfoFromTurnkey(Line, PlatformName, Unused, SdkInfo) == false)
+// 						{
+// 							continue;
+// 						}
+// 
+// 						// check if we had already set a ManualSDK - and don't set it again. Because of the way AutoSDKs are activated in the editor after the first call to Turnkey,
+// 						// future calls to Turnkey will inherit the AutoSDK env vars, and it won't be able to determine the manual SDK versions anymore. If we use the editor to
+// 						// install an SDK via Turnkey, it will directly update the installed version based on the result of that command, not this Update operation
+// 
+// 						FString OriginalManualInstallValue = DataDrivenPlatforms[*PlatformName].SdkInfo.InstalledVersion;
+// 
+// 						// set it into the platform
+// 						DataDrivenPlatforms[*PlatformName].SdkInfo = SdkInfo;
+// 
+// 						// restore the original installed version if it set after the first time
+// 						if (OriginalManualInstallValue.Len() > 0)
+// 						{
+// 							DataDrivenPlatforms[*PlatformName].SdkInfo.InstalledVersion = OriginalManualInstallValue;
+// 						}
+// 
+// 
+// 						UE_LOG(LogTemp, Log, TEXT("[TEST] Turnkey Platform: %s - %d, Installed: %s, AudoSDK: %s, Allowed: %s-%s"), *PlatformName, (int)SdkInfo.Status, *SdkInfo.InstalledVersion,
+// 							*SdkInfo.AutoSDKVersion, *SdkInfo.MinAllowedVersion, *SdkInfo.MaxAllowedVersion);
+// 					}
+// 				}
+// 			}
+// 			else
+// 			{
+// 				for (auto& It : DataDrivenPlatforms)
+// 				{
+// 					It.Value.SdkInfo.Status = DDPIPlatformSdkStatus::Error;
+// 					It.Value.SdkInfo.SdkErrorInformation = FText::Format(NSLOCTEXT("Turnkey", "TurnkeyError_ReturnedError", "Turnkey returned an error, code {0}"), { ExitCode });
+// 
+// 					// @todo turnkey error description!
+// 				}
+// 			}
+// 
+// 
+// 			for (auto& It : DataDrivenPlatforms)
+// 			{
+// 				if (It.Value.SdkInfo.Status == DDPIPlatformSdkStatus::Querying)
+// 				{
+// 					if (It.Value.bIsFakePlatform)
+// 					{
+// 						It.Value.SdkInfo.Status = DDPIPlatformSdkStatus::Unknown;
+// 					}
+// 					else
+// 					{
+// 						It.Value.SdkInfo.Status = DDPIPlatformSdkStatus::Error;
+// 						It.Value.SdkInfo.SdkErrorInformation = NSLOCTEXT("Turnkey", "TurnkeyError_NotReturned", "The platform's Sdk status was not returned from Turnkey");
+// 					}
+// 				}
+// 			}
+// 
+// 			// cleanup
+// 			delete TurnkeyProcess;
+// 			IFileManager::Get().Delete(*ReportFilename);
+// 		});
+// 	});
+// 
+// 	// run it
+// 	TurnkeyProcess->Launch();
+// }
+// 
+// void FTurnkeySupportModule::UpdateSdkInfoForDevices(TArray<FString> DeviceIds) const
+// {
+// 	FString Command, BaseCommandline, ReportFilename;
+// 	PrepForTurnkeyReport(Command, BaseCommandline, ReportFilename);
+// 
+// 	// the platform part of the Id may need to be converted to be turnkey (ie UBT) proper
+// 
+// 	FString Commandline = BaseCommandline + FString(TEXT(" -Device=")) + FString::JoinBy(PlatformDeviceIds, TEXT("+"), [](FString Id) { return ConvertToUATDeviceId(Id); });
+// 
+// 	UE_LOG(LogInit, Log, TEXT("Running Turnkey SDK detection: '%s %s'"), *Command, *Commandline);
+// 
+// 	{
+// 		FScopeLock Lock(&DDPILocker);
+// 
+// 		// set status to querying
+// 		FDDPISdkInfo DefaultInfo;
+// 		DefaultInfo.Status = DDPIPlatformSdkStatus::Querying;
+// 		for (const FString& Id : PlatformDeviceIds)
+// 		{
+// 			DeviceIdToInfo(Id).PerDeviceStatus.Add(ConvertToDDPIDeviceId(Id), DefaultInfo);
+// 		}
+// 	}
+// 
+// 	FMonitoredProcess* TurnkeyProcess = new FMonitoredProcess(Command, Commandline, true, false);
+// 	TurnkeyProcess->OnCompleted().BindLambda([ReportFilename, TurnkeyProcess, PlatformDeviceIds](int32 ExitCode)
+// 	{
+// 		AsyncTask(ENamedThreads::GameThread, [ReportFilename, TurnkeyProcess, PlatformDeviceIds, ExitCode]()
+// 		{
+// 			FScopeLock Lock(&DDPILocker);
+// 
+// 			if (ExitCode == 0 || ExitCode == 10)
+// 			{
+// 				TArray<FString> Contents;
+// 				if (FFileHelper::LoadFileToStringArray(Contents, *ReportFilename))
+// 				{
+// 					for (FString& Line : Contents)
+// 					{
+// 						FString PlatformName, DDPIDeviceId;
+// 						FDDPISdkInfo SdkInfo;
+// 						if (GetSdkInfoFromTurnkey(Line, PlatformName, DDPIDeviceId, SdkInfo) == false)
+// 						{
+// 							continue;
+// 						}
+// 
+// 						// skip over non-device lines
+// 						if (DDPIDeviceId.Len() == 0)
+// 						{
+// 							continue;
+// 						}
+// 
+// 						UE_LOG(LogTemp, Log, TEXT("Turnkey Device: %s"), *Line);
+// 
+// 						DeviceIdToInfo(DDPIDeviceId).PerDeviceStatus[DDPIDeviceId] = SdkInfo;
+// 
+// 						UE_LOG(LogTemp, Log, TEXT("[TEST] Turnkey Device: %s - %d, Installed: %s, Allowed: %s-%s"), *DDPIDeviceId, (int)SdkInfo.Status, *SdkInfo.InstalledVersion,
+// 							*SdkInfo.MinAllowedVersion, *SdkInfo.MaxAllowedVersion);
+// 					}
+// 				}
+// 			}
+// 
+// 			for (const FString& Id : PlatformDeviceIds)
+// 			{
+// 				FDataDrivenPlatformInfo& Info = DeviceIdToInfo(Id);
+// 
+// 				FDDPISdkInfo& SdkInfo = Info.PerDeviceStatus[ConvertToDDPIDeviceId(Id)];
+// 				if (SdkInfo.Status == DDPIPlatformSdkStatus::Querying)
+// 				{
+// 					SdkInfo.Status = DDPIPlatformSdkStatus::Error;
+// 					SdkInfo.SdkErrorInformation = NSLOCTEXT("Turnkey", "TurnkeyError_DeviceNotReturned", "A device's Sdk status was not returned from Turnkey");
+// 				}
+// 			}
+// 
+// 			// cleanup
+// 			delete TurnkeyProcess;
+// 			IFileManager::Get().Delete(*ReportFilename);
+// 		});
+// 	});
+// 
+// 	// run it
+// 	TurnkeyProcess->Launch();
+// }
+// 
