- Fix Turnkey to launch with DeviceID.Name, instead of Device.Name (which is the display string it turns out)

diff --git a/Engine/Source/Editor/TurnkeySupport/Private/TurnkeySupportModule.cpp b/Engine/Source/Editor/TurnkeySupport/Private/TurnkeySupportModule.cpp
--- a/Engine/Source/Editor/TurnkeySupport/Private/TurnkeySupportModule.cpp
+++ b/Engine/Source/Editor/TurnkeySupport/Private/TurnkeySupportModule.cpp
@@ -563,7 +563,7 @@ TSharedRef<FUICommandList> FTurnkeySupportCommands::ActionList = MakeShareable(n
 
 
 
-static void TurnkeyInstallSdk(FString PlatformName, bool bPreferFull, bool bForceInstall, const TCHAR* DeviceName = nullptr)
+static void TurnkeyInstallSdk(FString PlatformName, bool bPreferFull, bool bForceInstall, FString DeviceId)
 {
 	//	FString CommandLine = FString::Printf(TEXT("Turnkey -command=InstallSdk -platform=%s -BestAvailable -AllowAutoSdk -EditorIO -noturnkeyvariables -utf8output -WaitForUATMutex"), *PlatformName);
 
@@ -574,11 +574,11 @@ static void TurnkeyInstallSdk(FString PlatformName, bool bPreferFull, bool bForc
 	}
 	if (bForceInstall)
 	{
-		OptionalOptions += DeviceName != nullptr ? TEXT(" -ForceDeviceInstall") : TEXT(" -ForceSdkInstall");
+		OptionalOptions += DeviceId.Len() > 0 ? TEXT(" -ForceDeviceInstall") : TEXT(" -ForceSdkInstall");
 	}
-	if (DeviceName != nullptr)
+	if (DeviceId.Len() > 0)
 	{
-		OptionalOptions += FString::Printf(TEXT(" -Device=%s"), DeviceName);
+		OptionalOptions += FString::Printf(TEXT(" -Device=%s"), *DeviceId);
 	}
 
 	FString CommandLine = FString::Printf(TEXT("Turnkey -command=VerifySdk -UpdateIfNeeded -platform=%s %s -EditorIO -noturnkeyvariables -utf8output -WaitForUATMutex"), *PlatformName, *OptionalOptions);
@@ -802,14 +802,14 @@ static void MakeTurnkeyPlatformMenu(FMenuBuilder& MenuBuilder, FName IniPlatform
 						LOCTEXT("Turnkey_RepairDevice", "Repair Device as Needed"),
 						LOCTEXT("TurnkeyTooltip_RepairDevice", "Perform any fixup that may be needed on this device. If up to date already, nothing will be done."),
 						FSlateIcon(),
-						FExecuteAction::CreateStatic(TurnkeyInstallSdk, UBTPlatformString, false, false, *DeviceName)
+						FExecuteAction::CreateStatic(TurnkeyInstallSdk, UBTPlatformString, false, false, DeviceId)
 					);
 
 					SubMenuBuilder.AddMenuEntry(
 						LOCTEXT("Turnkey_ForceRepairDevice", "Force Repair Device"),
 						LOCTEXT("TurnkeyTooltip_ForceRepairDevice", "Force repairing anything on the device needed (update firmware, etc). Will perform all steps possible, even if not needed."),
 						FSlateIcon(),
-						FExecuteAction::CreateStatic(TurnkeyInstallSdk, UBTPlatformString, true, false, *DeviceName)
+						FExecuteAction::CreateStatic(TurnkeyInstallSdk, UBTPlatformString, true, false, DeviceId)
 					);
 				})
 			);
@@ -830,7 +830,7 @@ static void MakeTurnkeyPlatformMenu(FMenuBuilder& MenuBuilder, FName IniPlatform
 		FText::GetEmpty()
 	);
 
-	const TCHAR* NoDevice = nullptr;
+	FString NoDevice;
 	if (SdkInfo.Status == ETurnkeyPlatformSdkStatus::OutOfDate)
 	{
 		MenuBuilder.AddMenuEntry(
@@ -967,7 +967,7 @@ static void LaunchOnDevice(const FString& DeviceId, const FString& DeviceName, b
 			// if we want to check device flash before we start cooking, kick it off now. we could delay this 
 			if (bUseTurnkey)
 			{
-				FString CommandLine = FString::Printf(TEXT("Turnkey -command=VerifySdk -UpdateIfNeeded -platform=%s -EditorIO -noturnkeyvariables -device=%s -utf8output -WaitForUATMutex"), *UBTPlatformName, *DeviceName);
+				FString CommandLine = FString::Printf(TEXT("Turnkey -command=VerifySdk -UpdateIfNeeded -platform=%s -EditorIO -noturnkeyvariables -device=%s -utf8output -WaitForUATMutex"), *UBTPlatformName, *TargetDeviceId.GetDeviceName());
 				FText TaskName = LOCTEXT("VerifyingSDK", "Verifying SDK and Device");
 
 				IUATHelperModule::Get().CreateUatTask(CommandLine, FText::FromString(IniPlatformName), TaskName, TaskName, FEditorStyle::GetBrush(TEXT("MainFrame.PackageProject")),
