- Added some error handling for an external studio to track down issue. Also, no longer will crash in the error case

diff --git a/Engine/Source/Editor/TurnkeySupport/Private/TurnkeySupportModule.cpp b/Engine/Source/Editor/TurnkeySupport/Private/TurnkeySupportModule.cpp
--- a/Engine/Source/Editor/TurnkeySupport/Private/TurnkeySupportModule.cpp
+++ b/Engine/Source/Editor/TurnkeySupport/Private/TurnkeySupportModule.cpp
@@ -1439,7 +1439,7 @@ void FTurnkeySupportModule::UpdateSdkInfo()
 	// get status for all platforms
 	FString Commandline = BaseCommandline + TEXT(" -platform=all");
 
-	UE_LOG(LogInit, Log, TEXT("Running Turnkey SDK detection: '%s %s'"), *Command, *Commandline);
+	UE_LOG(LogTurnkeySupport, Log, TEXT("Running Turnkey SDK detection: '%s %s'"), *Command, *Commandline);
 
 	{
 		FScopeLock Lock(&GTurnkeySection);
@@ -1468,7 +1468,7 @@ void FTurnkeySupportModule::UpdateSdkInfo()
 				{
 					for (FString& Line : Contents)
 					{
-						UE_LOG(LogTemp, Log, TEXT("Turnkey Platform: %s"), *Line);
+						UE_LOG(LogTurnkeySupport, Log, TEXT("Turnkey Platform: %s"), *Line);
 
 						// parse a Turnkey line
 						FName PlatformName;
@@ -1479,6 +1479,13 @@ void FTurnkeySupportModule::UpdateSdkInfo()
 							continue;
 						}
 
+						// we received a platform from UAT that we don't know about in the editor. this shouldn't happen, because UAT and Editor should
+						// have the same idea of platforms, but we will print out an error and continue
+						if (PerPlatformSdkInfo.Contains(PlatformName))
+						{
+							UE_LOG(LogTurnkeySupport, Error, TEXT("Received platform %s from Turnkey, but the engine doesn't know about it."), *PlatformName.ToString());
+						}
+
 						// check if we had already set a ManualSDK - and don't set it again. Because of the way AutoSDKs are activated in the editor after the first call to Turnkey,
 						// future calls to Turnkey will inherit the AutoSDK env vars, and it won't be able to determine the manual SDK versions anymore. If we use the editor to
 						// install an SDK via Turnkey, it will directly update the installed version based on the result of that command, not this Update operation
@@ -1495,7 +1502,7 @@ void FTurnkeySupportModule::UpdateSdkInfo()
 						}
 
 
-						UE_LOG(LogTemp, Log, TEXT("[TEST] Turnkey Platform: %s - %d, Installed: %s, AudoSDK: %s, Allowed: %s-%s"), *PlatformName.ToString(), (int)SdkInfo.Status, *SdkInfo.InstalledVersion,
+						UE_LOG(LogTurnkeySupport, Log, TEXT("[TEST] Turnkey Platform: %s - %d, Installed: %s, AudoSDK: %s, Allowed: %s-%s"), *PlatformName.ToString(), (int)SdkInfo.Status, *SdkInfo.InstalledVersion,
 							*SdkInfo.AutoSDKVersion, *SdkInfo.MinAllowedVersion, *SdkInfo.MaxAllowedVersion);
 					}
 				}
@@ -1548,7 +1555,7 @@ void FTurnkeySupportModule::UpdateSdkInfoForDevices(TArray<FString> PlatformDevi
 
 	FString Commandline = BaseCommandline + FString(TEXT(" -Device=")) + FString::JoinBy(PlatformDeviceIds, TEXT("+"), [](FString Id) { return ConvertToUATDeviceId(Id); });
 
-	UE_LOG(LogInit, Log, TEXT("Running Turnkey SDK detection: '%s %s'"), *Command, *Commandline);
+	UE_LOG(LogTurnkeySupport, Log, TEXT("Running Turnkey SDK detection: '%s %s'"), *Command, *Commandline);
 
 	{
 		FScopeLock Lock(&GTurnkeySection);
@@ -1590,11 +1597,18 @@ void FTurnkeySupportModule::UpdateSdkInfoForDevices(TArray<FString> PlatformDevi
 							continue;
 						}
 
-						UE_LOG(LogTemp, Log, TEXT("Turnkey Device: %s"), *Line);
+						// we received a device from UAT that we don't know about in the editor. this should never happen since we pass a list of devices to Turnkey, 
+						// so this is a logic error
+						if (PerDeviceSdkInfo.Contains(DDPIDeviceId))
+						{
+							UE_LOG(LogTurnkeySupport, Error, TEXT("Received DeviceId %s from Turnkey, but the engine doesn't know about it."), *DDPIDeviceId);
+						}
+
+						UE_LOG(LogTurnkeySupport, Log, TEXT("Turnkey Device: %s"), *Line);
 
 						PerDeviceSdkInfo[DDPIDeviceId] = SdkInfo;
 
-						UE_LOG(LogTemp, Log, TEXT("[TEST] Turnkey Device: %s - %d, Installed: %s, Allowed: %s-%s"), *DDPIDeviceId, (int)SdkInfo.Status, *SdkInfo.InstalledVersion,
+						UE_LOG(LogTurnkeySupport, Log, TEXT("[TEST] Turnkey Device: %s - %d, Installed: %s, Allowed: %s-%s"), *DDPIDeviceId, (int)SdkInfo.Status, *SdkInfo.InstalledVersion,
 							*SdkInfo.MinAllowedVersion, *SdkInfo.MaxAllowedVersion);
 					}
 				}
