- Downgraded an Error log to Log, since it isn't really an error, since you may be running the editor without a plaform DDPI.ini file, even tho UBT/UAT had the platform compiled in

diff --git a/Engine/Source/Editor/TurnkeySupport/Private/TurnkeySupportModule.cpp b/Engine/Source/Editor/TurnkeySupport/Private/TurnkeySupportModule.cpp
--- a/Engine/Source/Editor/TurnkeySupport/Private/TurnkeySupportModule.cpp
+++ b/Engine/Source/Editor/TurnkeySupport/Private/TurnkeySupportModule.cpp
@@ -1479,11 +1479,11 @@ void FTurnkeySupportModule::UpdateSdkInfo()
 							continue;
 						}
 
-						// we received a platform from UAT that we don't know about in the editor. this shouldn't happen, because UAT and Editor should
-						// have the same idea of platforms, but we will print out an error and continue
+						// we received a platform from UAT that we don't know about in the editor. this can happen if you have a UBT/UAT that was compiled with platform access
+						// but then you are running without that platform synced. skip this platform and move on
 						if (PerPlatformSdkInfo.Contains(PlatformName))
 						{
-							UE_LOG(LogTurnkeySupport, Error, TEXT("Received platform %s from Turnkey, but the engine doesn't know about it."), *PlatformName.ToString());
+							UE_LOG(LogTurnkeySupport, Log, TEXT("Received platform %s from Turnkey, but the engine doesn't know about it. Skipping..."), *PlatformName.ToString());
 						}
 
 						// check if we had already set a ManualSDK - and don't set it again. Because of the way AutoSDKs are activated in the editor after the first call to Turnkey,
