Fixed the logic that showed a Turnkey warning mistakenly

diff --git a/Engine/Source/Editor/TurnkeySupport/Private/TurnkeySupportModule.cpp b/Engine/Source/Editor/TurnkeySupport/Private/TurnkeySupportModule.cpp
--- a/Engine/Source/Editor/TurnkeySupport/Private/TurnkeySupportModule.cpp
+++ b/Engine/Source/Editor/TurnkeySupport/Private/TurnkeySupportModule.cpp
@@ -1599,7 +1599,7 @@ void FTurnkeySupportModule::UpdateSdkInfoForDevices(TArray<FString> PlatformDevi
 
 						// we received a device from UAT that we don't know about in the editor. this should never happen since we pass a list of devices to Turnkey, 
 						// so this is a logic error
-						if (PerDeviceSdkInfo.Contains(DDPIDeviceId))
+						if (!PerDeviceSdkInfo.Contains(DDPIDeviceId))
 						{
 							UE_LOG(LogTurnkeySupport, Error, TEXT("Received DeviceId %s from Turnkey, but the engine doesn't know about it."), *DDPIDeviceId);
 						}
