Beginning the move of Editor stuff pulled out of Turnkey, to get it able to move to Developer module for UFE to use it

diff --git a/Engine/Source/Editor/TurnkeySupport/Private/TurnkeyEditorSupport.cpp b/Engine/Source/Editor/TurnkeySupport/Private/TurnkeyEditorSupport.cpp
new file mode 100644
--- /dev/null
+++ b/Engine/Source/Editor/TurnkeySupport/Private/TurnkeyEditorSupport.cpp
@@ -0,0 +1,162 @@
+// Copyright Epic Games, Inc. All Rights Reserved.
+
+#include "TurnkeyEditorSupport.h"
+
+#if WITH_EDITOR
+#include "UnrealEdMisc.h"
+#include "Editor/UnrealEdEngine.h"
+#include "UnrealEdGlobals.h"
+#include "PlatformInfo.h"
+#include "InstalledPlatformInfo.h"
+#include "Misc/MessageDialog.h"
+#include "IUATHelperModule.h"
+#include "Interfaces/IProjectTargetPlatformEditorModule.h"
+#include "Dialogs/Dialogs.h"
+#include "Async/Async.h"
+#include "GameProjectGenerationModule.h"
+
+#endif
+
+#define LOCTEXT_NAMESPACE "FTurnkeyEditorSupport"
+
+FString FTurnkeyEditorSupport::GetUATOptions()
+{
+#if WITH_EDITOR
+	return FUnrealEdMisc::Get().GetExecutableForCommandlets();
+#else
+	return TEXT("");
+#endif
+}
+
+void FTurnkeyEditorSupport::AddEditorOptions(FMenuBuilder& MenuBuilder)
+{
+	FModuleManager::LoadModuleChecked<IProjectTargetPlatformEditorModule>("ProjectTargetPlatformEditor").AddOpenProjectTargetPlatformEditorMenuItem(MenuBuilder);
+}
+
+void FTurnkeyEditorSupport::LaunchRunningMap(const FString& DeviceId, const FString& DeviceName, bool bUseTurnkey)
+{
+#if WITH_EDITOR
+	FTargetDeviceId TargetDeviceId;
+	if (FTargetDeviceId::Parse(DeviceId, TargetDeviceId))
+	{
+		const PlatformInfo::FTargetPlatformInfo* const PlatformInfo = PlatformInfo::FindPlatformInfo(*TargetDeviceId.GetPlatformName());
+		FString UBTPlatformName = PlatformInfo->DataDrivenPlatformInfo->UBTPlatformString;
+		FString IniPlatformName = PlatformInfo->IniPlatformName.ToString();
+
+		check(PlatformInfo);
+
+		if (FInstalledPlatformInfo::Get().IsPlatformMissingRequiredFile(UBTPlatformName))
+		{
+			if (!FInstalledPlatformInfo::OpenInstallerOptions())
+			{
+				FMessageDialog::Open(EAppMsgType::Ok, LOCTEXT("MissingPlatformFilesLaunch", "Missing required files to launch on this platform."));
+			}
+			return;
+		}
+
+		if (FModuleManager::LoadModuleChecked<IProjectTargetPlatformEditorModule>("ProjectTargetPlatformEditor").ShowUnsupportedTargetWarning(*TargetDeviceId.GetPlatformName()))
+		{
+			GUnrealEd->CancelPlayingViaLauncher();
+
+			FRequestPlaySessionParams::FLauncherDeviceInfo DeviceInfo;
+			DeviceInfo.DeviceId = DeviceId;
+			DeviceInfo.DeviceName = DeviceName;
+			// @todo turnkey: we set this to false because we will kick off a Turnkey run before cooking, etc, to get an early warning. however, if it's too difficult
+			// to get an error back from CreateUatTask, then we should set this to bUseTurnkey and remove the block below, and let the code in FLauncherWorker::CreateAndExecuteTasks handle it
+			DeviceInfo.bUpdateDeviceFlash = false;
+
+			FRequestPlaySessionParams SessionParams;
+			SessionParams.SessionDestination = EPlaySessionDestinationType::Launcher;
+			SessionParams.LauncherTargetDevice = DeviceInfo;
+
+			// if we want to check device flash before we start cooking, kick it off now. we could delay this 
+			if (bUseTurnkey)
+			{
+				FString CommandLine = FString::Printf(TEXT("Turnkey -command=VerifySdk -UpdateIfNeeded -platform=%s -EditorIO -noturnkeyvariables -device=%s -utf8output -WaitForUATMutex"), *UBTPlatformName, *TargetDeviceId.GetDeviceName());
+				FText TaskName = LOCTEXT("VerifyingSDK", "Verifying SDK and Device");
+
+				IUATHelperModule::Get().CreateUatTask(CommandLine, FText::FromString(IniPlatformName), TaskName, TaskName, FEditorStyle::GetBrush(TEXT("MainFrame.PackageProject")),
+					[SessionParams](FString Result, double)
+					{
+						// unfortunate string comparison for success
+						bool bWasSuccessful = Result == TEXT("Completed");
+						AsyncTask(ENamedThreads::GameThread, [SessionParams, bWasSuccessful]()
+							{
+								if (bWasSuccessful)
+								{
+									GUnrealEd->RequestPlaySession(SessionParams);
+								}
+								else
+								{
+									TSharedRef<SWindow> Win = OpenMsgDlgInt_NonModal(EAppMsgType::YesNo, LOCTEXT("SDKCheckFailed", "SDK Verification failed. Would you like to attempt the Launch On anyway?"), LOCTEXT("SDKCheckFailedTitle", "SDK Verification"),
+										FOnMsgDlgResult::CreateLambda([SessionParams](const TSharedRef<SWindow>&, EAppReturnType::Type Choice)
+											{
+												if (Choice == EAppReturnType::Yes)
+												{
+													GUnrealEd->RequestPlaySession(SessionParams);
+												}
+											}));
+									Win->ShowWindow();
+								}
+							});
+					});
+			}
+			else
+			{
+				GUnrealEd->RequestPlaySession(SessionParams);
+			}
+		}
+	}
+#endif
+}
+
+bool FTurnkeyEditorSupport::DoesProjectHaveCode()
+{
+#if WITH_EDITOR
+	FGameProjectGenerationModule& GameProjectModule = FModuleManager::LoadModuleChecked<FGameProjectGenerationModule>(TEXT("GameProjectGeneration"));
+	return GameProjectModule.Get().ProjectHasCodeFiles();
+#else
+	unimplemented;
+#endif
+}
+
+void FTurnkeyEditorSupport::RunUAT(const FString& CommandLine, const FText& PlatformDisplayName, const FText& TaskName, const FText& TaskShortName, const FSlateBrush* TaskIcon, TFunction<void(FString, double)> ResultCallback)
+{
+#if WITH_EDITOR
+	IUATHelperModule::Get().CreateUatTask(CommandLine, PlatformDisplayName, TaskName, TaskShortName, TaskIcon, ResultCallback);
+#else
+	unimplemented;
+#endif
+}
+
+
+
+
+bool FTurnkeyEditorSupport::ShowOKCancelDialog(FText Message, FText Title)
+{
+#if WITH_EDITOR
+	FSuppressableWarningDialog::FSetupInfo Info(Message, Title, TEXT("TurkeyEditorDialog"));
+
+	Info.ConfirmText = LOCTEXT("TurnkeyDialog_Confirm", "Continue");
+	Info.CancelText = LOCTEXT("TurnkeyDialogK_Cancel", "Cancel");
+	FSuppressableWarningDialog Dialog(Info);
+
+	return Dialog.ShowModal() != FSuppressableWarningDialog::EResult::Cancel;
+#else
+	unimplemented();
+	return false;
+#endif
+}
+
+
+bool FTurnkeyEditorSupport::CheckSupportedPlatforms(FName IniPlatformName)
+{
+#if WITH_EDITOR
+	return FModuleManager::LoadModuleChecked<IProjectTargetPlatformEditorModule>("ProjectTargetPlatformEditor").ShowUnsupportedTargetWarning(IniPlatformName);
+#endif
+
+	return true;
+}
+
+
+#undef LOCTEXT_NAMESPACE
diff --git a/Engine/Source/Editor/TurnkeySupport/Private/TurnkeyEditorSupport.h b/Engine/Source/Editor/TurnkeySupport/Private/TurnkeyEditorSupport.h
new file mode 100644
--- /dev/null
+++ b/Engine/Source/Editor/TurnkeySupport/Private/TurnkeyEditorSupport.h
@@ -0,0 +1,23 @@
+// Copyright Epic Games, Inc. All Rights Reserved.
+
+#pragma once
+
+#include "CoreMinimal.h"
+
+
+class FTurnkeyEditorSupport
+{
+public:
+
+	static FString GetUATOptions();
+
+	static void LaunchRunningMap(const FString& DeviceId, const FString& DeviceName, bool bUseTurnkey);
+	static void AddEditorOptions(FMenuBuilder& MenuBuilder);
+
+	static bool DoesProjectHaveCode();
+	static void RunUAT(const FString& CommandLine, const FText& PlatformDisplayName, const FText& TaskName, const FText& TaskShortName, const FSlateBrush* TaskIcon, TFunction<void(FString, double)> ResultCallback=TFunction<void(FString, double)>());
+
+	static bool ShowOKCancelDialog(FText Message, FText Title);
+	static bool CheckSupportedPlatforms(FName IniPlatformName);
+
+};
diff --git a/Engine/Source/Editor/TurnkeySupport/Private/TurnkeySupportModule.cpp b/Engine/Source/Editor/TurnkeySupport/Private/TurnkeySupportModule.cpp
--- a/Engine/Source/Editor/TurnkeySupport/Private/TurnkeySupportModule.cpp
+++ b/Engine/Source/Editor/TurnkeySupport/Private/TurnkeySupportModule.cpp
@@ -2,12 +2,12 @@
 
 #include "TurnkeySupportModule.h"
 #include "TurnkeySupport.h"
+#include "TurnkeyEditorSupport.h"
 
 #include "SlateOptMacros.h"
 #include "Widgets/Images/SImage.h"
 #include "Widgets/Layout/SBox.h"
 #include "Widgets/Input/SButton.h"
-#include "GameProjectGenerationModule.h"
 #include "MessageLogModule.h"
 #include "EditorStyleSet.h"
 #include "Framework/Notifications/NotificationManager.h"
@@ -15,12 +15,8 @@
 #include "Widgets/Input/SEditableTextBox.h"
 #include "Misc/MessageDialog.h"
 #include "HAL/PlatformFileManager.h"
-#include "UnrealEdMisc.h"
 #include "ITargetDeviceServicesModule.h"
-#include "Interfaces/IProjectTargetPlatformEditorModule.h"
-#include "Dialogs/Dialogs.h"
 #include "Misc/DataDrivenPlatformInfoRegistry.h"
-#include "IUATHelperModule.h"
 #include "ISettingsEditorModule.h"
 #include "Async/Async.h"
 #include "Misc/FileHelper.h"
@@ -40,15 +36,15 @@
 #include "InstalledPlatformInfo.h"
 #include "Settings/EditorExperimentalSettings.h"
 #include "CookerSettings.h"
-#include "Editor/UnrealEdEngine.h"
-#include "UnrealEdGlobals.h"
 #include "Framework/Commands/Commands.h"
 #include "Framework/Docking/TabManager.h"
 #include "Framework/MultiBox/MultiBoxBuilder.h"
 #include "Developer/DerivedDataCache/Public/DerivedDataCacheInterface.h"
-
 #include "Misc/MonitoredProcess.h"
 
+#include "Misc/App.h"
+#include "Framework/Application/SlateApplication.h"
+
 
 
 
@@ -88,16 +84,9 @@ protected:
 			Args.Add(TEXT("DisplayName"), FText::FromName(IniPlatformName));
 			FText WarningText = FText::Format(LOCTEXT("BadSDK_Message", "The SDK for {DisplayName} is not installed properly, which is needed to generate data. Check the SDK section of the Launch On menu in the main toolbar to update SDK.\n\nWould you like to attempt to continue anyway?"), Args);
 
-			FSuppressableWarningDialog::FSetupInfo Info(
-				WarningText,
-				LOCTEXT("BadSDK_Title", "SDK Not Setup"),
-				TEXT("BadSDKDialog")
-			);
-			Info.ConfirmText = LOCTEXT("BadSDK_Confirm", "Continue");
-			Info.CancelText = LOCTEXT("BadSDK_Cancel", "Cancel");
-			FSuppressableWarningDialog BadSDKgDialog(Info);
+			bool bClickedOK = FTurnkeyEditorSupport::ShowOKCancelDialog(WarningText, LOCTEXT("BadSDK_Title", "SDK Not Setup"));
+			return bClickedOK;
 
-			return BadSDKgDialog.ShowModal() != FSuppressableWarningDialog::EResult::Cancel;
 		}
 
 		return true;
@@ -151,10 +140,7 @@ protected:
 				{
 					FText Reason;
 
-					FGameProjectGenerationModule& GameProjectModule = FModuleManager::LoadModuleChecked<FGameProjectGenerationModule>(TEXT("GameProjectGeneration"));
-					bool bProjectHasCode = GameProjectModule.Get().ProjectHasCodeFiles();
-
-					if (TargetPlatform->RequiresTempTarget(bProjectHasCode, ConfigurationInfo.Configuration, false, Reason))
+					if (TargetPlatform->RequiresTempTarget(FTurnkeyEditorSupport::DoesProjectHaveCode(), ConfigurationInfo.Configuration, false, Reason))
 					{
 						UE_LOG(LogTurnkeySupport, Log, TEXT("Project requires temp target (%s)"), *Reason.ToString());
 						BaseDir = ProjectDir;
@@ -244,7 +230,7 @@ public:
 				return;
 			}
 
-			if (!FModuleManager::LoadModuleChecked<IProjectTargetPlatformEditorModule>("ProjectTargetPlatformEditor").ShowUnsupportedTargetWarning(IniPlatformName))
+			if (!FTurnkeyEditorSupport::CheckSupportedPlatforms(IniPlatformName))
 			{
 				return;
 			}
@@ -262,7 +248,12 @@ public:
 
 		// set locations to engine and project
 		{
-			BuildCookRunParams += FString::Printf(TEXT(" -project=\"%s\" -ue4exe=\"%s\""), *ProjectPath, *FUnrealEdMisc::Get().GetExecutableForCommandlets());
+			BuildCookRunParams += FString::Printf(TEXT(" -project=\"%s\""), *ProjectPath);
+		}
+
+		// let the editor add options (-ue4exe in particular)
+		{
+			BuildCookRunParams += FString::Printf(TEXT(" %s"), *FTurnkeyEditorSupport::GetUATOptions());
 		}
 
 		// set the platform we are preparing content for
@@ -442,7 +433,7 @@ public:
 			*BuildCookRunParams
 		);
 
-		IUATHelperModule::Get().CreateUatTask(CommandLine, PlatformInfo->DisplayName, ContentPrepDescription, ContentPrepTaskName, ContentPrepIcon);
+		FTurnkeyEditorSupport::RunUAT(CommandLine, PlatformInfo->DisplayName, ContentPrepDescription, ContentPrepTaskName, ContentPrepIcon);
 	}
 
 	static void PackageBuildConfiguration(const PlatformInfo::FTargetPlatformInfo* Info, EProjectPackagingBuildConfigurations BuildConfiguration)
@@ -584,7 +575,7 @@ static void TurnkeyInstallSdk(FString PlatformName, bool bPreferFull, bool bForc
 	FString CommandLine = FString::Printf(TEXT("Turnkey -command=VerifySdk -UpdateIfNeeded -platform=%s %s -EditorIO -noturnkeyvariables -utf8output -WaitForUATMutex"), *PlatformName, *OptionalOptions);
 
 	FText TaskName = LOCTEXT("InstallingSdk", "Installing Sdk");
-	IUATHelperModule::Get().CreateUatTask(CommandLine, FText::FromString(PlatformName), TaskName, TaskName, FEditorStyle::GetBrush(TEXT("MainFrame.PackageProject")),
+	FTurnkeyEditorSupport::RunUAT(CommandLine, FText::FromString(PlatformName), TaskName, TaskName, FEditorStyle::GetBrush(TEXT("MainFrame.PackageProject")),
 		[PlatformName](FString, double)
 	{
 		AsyncTask(ENamedThreads::GameThread, [PlatformName]()
@@ -728,7 +719,7 @@ static void MakeTurnkeyPlatformMenu(FMenuBuilder& MenuBuilder, FName IniPlatform
 
 
 		MenuBuilder.BeginSection("BuildConfig", LOCTEXT("TurnkeySection_BuildConfig", "Binary Configuration"));
-			EProjectType ProjectType = FGameProjectGenerationModule::Get().ProjectHasCodeFiles() ? EProjectType::Code : EProjectType::Content;
+			EProjectType ProjectType = FTurnkeyEditorSupport::DoesProjectHaveCode() ? EProjectType::Code : EProjectType::Content;
 			TArray<EProjectPackagingBuildConfigurations> PackagingConfigurations = UProjectPackagingSettings::GetValidPackageConfigurations();
 
 			for (EProjectPackagingBuildConfigurations PackagingConfiguration : PackagingConfigurations)
@@ -886,122 +877,48 @@ static void MakeTurnkeyPlatformMenu(FMenuBuilder& MenuBuilder, FName IniPlatform
 
 bool CanLaunchOnDevice(const FString& DeviceName)
 {
-	if (!GUnrealEd->IsPlayingViaLauncher())
-	{
-		static TWeakPtr<ITargetDeviceProxyManager> DeviceProxyManagerPtr;
+	static TWeakPtr<ITargetDeviceProxyManager> DeviceProxyManagerPtr;
 
-		if (!DeviceProxyManagerPtr.IsValid())
+	if (!DeviceProxyManagerPtr.IsValid())
+	{
+		ITargetDeviceServicesModule* TargetDeviceServicesModule = FModuleManager::Get().LoadModulePtr<ITargetDeviceServicesModule>(TEXT("TargetDeviceServices"));
+		if (TargetDeviceServicesModule)
 		{
-			ITargetDeviceServicesModule* TargetDeviceServicesModule = FModuleManager::Get().LoadModulePtr<ITargetDeviceServicesModule>(TEXT("TargetDeviceServices"));
-			if (TargetDeviceServicesModule)
-			{
-				DeviceProxyManagerPtr = TargetDeviceServicesModule->GetDeviceProxyManager();
-			}
+			DeviceProxyManagerPtr = TargetDeviceServicesModule->GetDeviceProxyManager();
 		}
+	}
 
-		TSharedPtr<ITargetDeviceProxyManager> DeviceProxyManager = DeviceProxyManagerPtr.Pin();
-		if (DeviceProxyManager.IsValid())
+	TSharedPtr<ITargetDeviceProxyManager> DeviceProxyManager = DeviceProxyManagerPtr.Pin();
+	if (DeviceProxyManager.IsValid())
+	{
+		TSharedPtr<ITargetDeviceProxy> DeviceProxy = DeviceProxyManager->FindProxy(DeviceName);
+		if (DeviceProxy.IsValid() && DeviceProxy->IsConnected() && DeviceProxy->IsAuthorized())
 		{
-			TSharedPtr<ITargetDeviceProxy> DeviceProxy = DeviceProxyManager->FindProxy(DeviceName);
-			if (DeviceProxy.IsValid() && DeviceProxy->IsConnected() && DeviceProxy->IsAuthorized())
-			{
-				return true;
-			}
+			return true;
+		}
 
-			// check if this is an aggregate proxy
-			TArray<TSharedPtr<ITargetDeviceProxy>> Devices;
-			DeviceProxyManager->GetProxies(FName(*DeviceName), false, Devices);
+		// check if this is an aggregate proxy
+		TArray<TSharedPtr<ITargetDeviceProxy>> Devices;
+		DeviceProxyManager->GetProxies(FName(*DeviceName), false, Devices);
 
-			// returns true if the game can be launched al least on 1 device
-			for (auto DevicesIt = Devices.CreateIterator(); DevicesIt; ++DevicesIt)
+		// returns true if the game can be launched al least on 1 device
+		for (auto DevicesIt = Devices.CreateIterator(); DevicesIt; ++DevicesIt)
+		{
+			TSharedPtr<ITargetDeviceProxy> DeviceAggregateProxy = *DevicesIt;
+			if (DeviceAggregateProxy.IsValid() && DeviceAggregateProxy->IsConnected() && DeviceAggregateProxy->IsAuthorized())
 			{
-				TSharedPtr<ITargetDeviceProxy> DeviceAggregateProxy = *DevicesIt;
-				if (DeviceAggregateProxy.IsValid() && DeviceAggregateProxy->IsConnected() && DeviceAggregateProxy->IsAuthorized())
-				{
-					return true;
-				}
+				return true;
 			}
-
 		}
+
 	}
 
 	return false;
 }
 
-
 static void LaunchOnDevice(const FString& DeviceId, const FString& DeviceName, bool bUseTurnkey)
 {
-	FTargetDeviceId TargetDeviceId;
-	if (FTargetDeviceId::Parse(DeviceId, TargetDeviceId))
-	{
-		const PlatformInfo::FTargetPlatformInfo* const PlatformInfo = PlatformInfo::FindPlatformInfo(*TargetDeviceId.GetPlatformName());
-		FString UBTPlatformName = PlatformInfo->DataDrivenPlatformInfo->UBTPlatformString;
-		FString IniPlatformName = PlatformInfo->IniPlatformName.ToString();
-
-		check(PlatformInfo);
-
-		if (FInstalledPlatformInfo::Get().IsPlatformMissingRequiredFile(UBTPlatformName))
-		{
-			if (!FInstalledPlatformInfo::OpenInstallerOptions())
-			{
-				FMessageDialog::Open(EAppMsgType::Ok, LOCTEXT("MissingPlatformFilesLaunch", "Missing required files to launch on this platform."));
-			}
-			return;
-		}
-
-		if (FModuleManager::LoadModuleChecked<IProjectTargetPlatformEditorModule>("ProjectTargetPlatformEditor").ShowUnsupportedTargetWarning(*TargetDeviceId.GetPlatformName()))
-		{
-			GUnrealEd->CancelPlayingViaLauncher();
-
-			FRequestPlaySessionParams::FLauncherDeviceInfo DeviceInfo;
-			DeviceInfo.DeviceId = DeviceId;
-			DeviceInfo.DeviceName = DeviceName;
-			// @todo turnkey: we set this to false because we will kick off a Turnkey run before cooking, etc, to get an early warning. however, if it's too difficult
-			// to get an error back from CreateUatTask, then we should set this to bUseTurnkey and remove the block below, and let the code in FLauncherWorker::CreateAndExecuteTasks handle it
-			DeviceInfo.bUpdateDeviceFlash = false;
-
-			FRequestPlaySessionParams SessionParams;
-			SessionParams.SessionDestination = EPlaySessionDestinationType::Launcher;
-			SessionParams.LauncherTargetDevice = DeviceInfo;
-
-			// if we want to check device flash before we start cooking, kick it off now. we could delay this 
-			if (bUseTurnkey)
-			{
-				FString CommandLine = FString::Printf(TEXT("Turnkey -command=VerifySdk -UpdateIfNeeded -platform=%s -EditorIO -noturnkeyvariables -device=%s -utf8output -WaitForUATMutex"), *UBTPlatformName, *TargetDeviceId.GetDeviceName());
-				FText TaskName = LOCTEXT("VerifyingSDK", "Verifying SDK and Device");
-
-				IUATHelperModule::Get().CreateUatTask(CommandLine, FText::FromString(IniPlatformName), TaskName, TaskName, FEditorStyle::GetBrush(TEXT("MainFrame.PackageProject")),
-					[SessionParams](FString Result, double)
-				{
-					// unfortunate string comparison for success
-					bool bWasSuccessful = Result == TEXT("Completed");
-					AsyncTask(ENamedThreads::GameThread, [SessionParams, bWasSuccessful]()
-					{
-						if (bWasSuccessful)
-						{
-							GUnrealEd->RequestPlaySession(SessionParams);
-						}
-						else
-						{
-							TSharedRef<SWindow> Win = OpenMsgDlgInt_NonModal(EAppMsgType::YesNo, LOCTEXT("SDKCheckFailed", "SDK Verification failed. Would you like to attempt the Launch On anyway?"), LOCTEXT("SDKCheckFailedTitle", "SDK Verification"),
-								FOnMsgDlgResult::CreateLambda([SessionParams](const TSharedRef<SWindow>&, EAppReturnType::Type Choice)
-							{
-								if (Choice == EAppReturnType::Yes)
-								{
-									GUnrealEd->RequestPlaySession(SessionParams);
-								}
-							}));
-							Win->ShowWindow();
-						}
-					});
-				});
-			}
-			else
-			{
-				GUnrealEd->RequestPlaySession(SessionParams);
-			}
-		}
-	}
+	FTurnkeyEditorSupport::LaunchRunningMap(DeviceId, DeviceName, bUseTurnkey);
 }
 
 static void PrepareLaunchOn(FString DeviceId, FString DeviceName)
@@ -1092,7 +1009,6 @@ TSharedRef<SWidget> FTurnkeySupportModule::MakeTurnkeyMenu() const
 
 	// shared devices section
 	ITargetDeviceServicesModule* TargetDeviceServicesModule = static_cast<ITargetDeviceServicesModule*>(FModuleManager::Get().LoadModule(TEXT("TargetDeviceServices")));
-	IProjectTargetPlatformEditorModule& ProjectTargetPlatformEditorModule = FModuleManager::LoadModuleChecked<IProjectTargetPlatformEditorModule>("ProjectTargetPlatformEditor");
 
 	TArray<FString> DeviceIdsToQuery;
 
@@ -1305,7 +1221,7 @@ TSharedRef<SWidget> FTurnkeySupportModule::MakeTurnkeyMenu() const
 // 			FSlateIcon()
 // 		);
 
-		ProjectTargetPlatformEditorModule.AddOpenProjectTargetPlatformEditorMenuItem(MenuBuilder);
+		FTurnkeyEditorSupport::AddEditorOptions(MenuBuilder);
 	}
 	MenuBuilder.EndSection();
 
diff --git a/Engine/Source/Editor/TurnkeySupport/TurnkeySupport.Build.cs b/Engine/Source/Editor/TurnkeySupport/TurnkeySupport.Build.cs
--- a/Engine/Source/Editor/TurnkeySupport/TurnkeySupport.Build.cs
+++ b/Engine/Source/Editor/TurnkeySupport/TurnkeySupport.Build.cs
@@ -19,27 +19,29 @@ public class TurnkeySupport : ModuleRules
 				"RenderCore",
 				"Slate",
 				"SlateCore",
-				"EditorStyle",
-				"SourceControl",
-				"SourceControlWindows",
 				"TargetPlatform",
 				"DesktopPlatform",
-				"EditorFramework",
-				"UnrealEd",
 				"WorkspaceMenuStructure",
 				"MessageLog",
-				"UATHelper",
-				"TranslationEditor",
 				"Projects",
-				"DeviceProfileEditor",
-				"UndoHistory",
-				"Analytics",
 				"ToolMenus",
 				"LauncherServices",
 				"SettingsEditor"
 			}
 		);
 
+		if (Target.Type == TargetType.Editor)
+		{
+			PrivateDependencyModuleNames.AddRange(
+				new string[] {
+					"UnrealEd",
+					"UATHelper",
+					"EditorStyle",
+					"SourceControl",
+				}
+			);
+		}
+
 		PrivateIncludePathModuleNames.AddRange(
 			new string[] {
 				"AssetTools",
