- More stuff moved into Turnkey editor bits (now ready for UFE to use Turnkey, once some classes move out of UnrealEd)

diff --git a/Engine/Source/Editor/TurnkeySupport/Private/TurnkeyEditorSupport.cpp b/Engine/Source/Editor/TurnkeySupport/Private/TurnkeyEditorSupport.cpp
--- a/Engine/Source/Editor/TurnkeySupport/Private/TurnkeyEditorSupport.cpp
+++ b/Engine/Source/Editor/TurnkeySupport/Private/TurnkeyEditorSupport.cpp
@@ -1,6 +1,7 @@
 // Copyright Epic Games, Inc. All Rights Reserved.
 
 #include "TurnkeyEditorSupport.h"
+#include "Misc/AssertionMacros.h"
 
 #if WITH_EDITOR
 #include "UnrealEdMisc.h"
@@ -14,6 +15,9 @@
 #include "Dialogs/Dialogs.h"
 #include "Async/Async.h"
 #include "GameProjectGenerationModule.h"
+#include "ISettingsModule.h"
+#include "ISettingsEditorModule.h"
+#include "Settings/EditorExperimentalSettings.h"
 
 #endif
 
@@ -22,15 +26,50 @@
 FString FTurnkeyEditorSupport::GetUATOptions()
 {
 #if WITH_EDITOR
-	return FUnrealEdMisc::Get().GetExecutableForCommandlets();
+	FString Options;
+	Options += FString::Printf(TEXT(" -ue4exe=%s"), *FUnrealEdMisc::Get().GetExecutableForCommandlets());
+
+	int32 NumCookers = GetDefault<UEditorExperimentalSettings>()->MultiProcessCooking;
+	if (NumCookers > 0)
+	{
+		Options += FString::Printf(TEXT(" -NumCookersToSpawn=%d"), NumCookers);
+	}
+	return Options;
 #else
 	return TEXT("");
 #endif
 }
 
+
 void FTurnkeyEditorSupport::AddEditorOptions(FMenuBuilder& MenuBuilder)
 {
+#if WITH_EDITOR
+	MenuBuilder.AddSeparator();
+
+	MenuBuilder.AddMenuEntry(
+		LOCTEXT("OpenPackagingSettings", "Packaging Settings..."),
+		LOCTEXT("OpenPackagingSettings_ToolTip", "Opens the settings for project packaging."),
+		FSlateIcon(FEditorStyle::GetStyleSetName(), "DeviceDetails.TabIcon"),
+		FUIAction(FExecuteAction::CreateLambda([] { FModuleManager::LoadModuleChecked<ISettingsModule>("Settings").ShowViewer("Project", "Project", "Packaging"); }))
+	);
+
 	FModuleManager::LoadModuleChecked<IProjectTargetPlatformEditorModule>("ProjectTargetPlatformEditor").AddOpenProjectTargetPlatformEditorMenuItem(MenuBuilder);
+#endif
+}
+
+void FTurnkeyEditorSupport::PrepareToLaunchRunningMap(const FString& DeviceId, const FString& DeviceName)
+{
+#if WITH_EDITOR
+	ULevelEditorPlaySettings* PlaySettings = GetMutableDefault<ULevelEditorPlaySettings>();
+
+	PlaySettings->LastExecutedLaunchModeType = LaunchMode_OnDevice;
+	PlaySettings->LastExecutedLaunchDevice = DeviceId;
+	PlaySettings->LastExecutedLaunchName = DeviceName;
+
+	PlaySettings->PostEditChange();
+
+	PlaySettings->SaveConfig();
+#endif
 }
 
 void FTurnkeyEditorSupport::LaunchRunningMap(const FString& DeviceId, const FString& DeviceName, bool bUseTurnkey)
@@ -116,7 +155,8 @@ bool FTurnkeyEditorSupport::DoesProjectHaveCode()
 	FGameProjectGenerationModule& GameProjectModule = FModuleManager::LoadModuleChecked<FGameProjectGenerationModule>(TEXT("GameProjectGeneration"));
 	return GameProjectModule.Get().ProjectHasCodeFiles();
 #else
-	unimplemented;
+	unimplemented();
+	return false;
 #endif
 }
 
@@ -125,7 +165,7 @@ void FTurnkeyEditorSupport::RunUAT(const FString& CommandLine, const FText& Plat
 #if WITH_EDITOR
 	IUATHelperModule::Get().CreateUatTask(CommandLine, PlatformDisplayName, TaskName, TaskShortName, TaskIcon, ResultCallback);
 #else
-	unimplemented;
+	unimplemented();
 #endif
 }
 
@@ -148,6 +188,13 @@ bool FTurnkeyEditorSupport::ShowOKCancelDialog(FText Message, FText Title)
 #endif
 }
 
+void FTurnkeyEditorSupport::ShowRestartToast()
+{
+#if WITH_EDITOR
+	// show restart dialog
+	FModuleManager::GetModuleChecked<ISettingsEditorModule>("SettingsEditor").OnApplicationRestartRequired();
+#endif
+}
 
 bool FTurnkeyEditorSupport::CheckSupportedPlatforms(FName IniPlatformName)
 {
diff --git a/Engine/Source/Editor/TurnkeySupport/Private/TurnkeyEditorSupport.h b/Engine/Source/Editor/TurnkeySupport/Private/TurnkeyEditorSupport.h
--- a/Engine/Source/Editor/TurnkeySupport/Private/TurnkeyEditorSupport.h
+++ b/Engine/Source/Editor/TurnkeySupport/Private/TurnkeyEditorSupport.h
@@ -11,6 +11,7 @@ public:
 
 	static FString GetUATOptions();
 
+	static void PrepareToLaunchRunningMap(const FString& DeviceId, const FString& DeviceName);
 	static void LaunchRunningMap(const FString& DeviceId, const FString& DeviceName, bool bUseTurnkey);
 	static void AddEditorOptions(FMenuBuilder& MenuBuilder);
 
@@ -18,6 +19,7 @@ public:
 	static void RunUAT(const FString& CommandLine, const FText& PlatformDisplayName, const FText& TaskName, const FText& TaskShortName, const FSlateBrush* TaskIcon, TFunction<void(FString, double)> ResultCallback=TFunction<void(FString, double)>());
 
 	static bool ShowOKCancelDialog(FText Message, FText Title);
+	static void ShowRestartToast();
 	static bool CheckSupportedPlatforms(FName IniPlatformName);
 
 };
diff --git a/Engine/Source/Editor/TurnkeySupport/Private/TurnkeySupportModule.cpp b/Engine/Source/Editor/TurnkeySupport/Private/TurnkeySupportModule.cpp
--- a/Engine/Source/Editor/TurnkeySupport/Private/TurnkeySupportModule.cpp
+++ b/Engine/Source/Editor/TurnkeySupport/Private/TurnkeySupportModule.cpp
@@ -9,7 +9,6 @@
 #include "Widgets/Layout/SBox.h"
 #include "Widgets/Input/SButton.h"
 #include "MessageLogModule.h"
-#include "EditorStyleSet.h"
 #include "Framework/Notifications/NotificationManager.h"
 #include "Widgets/Notifications/SNotificationList.h"
 #include "Widgets/Input/SEditableTextBox.h"
@@ -17,14 +16,12 @@
 #include "HAL/PlatformFileManager.h"
 #include "ITargetDeviceServicesModule.h"
 #include "Misc/DataDrivenPlatformInfoRegistry.h"
-#include "ISettingsEditorModule.h"
 #include "Async/Async.h"
 #include "Misc/FileHelper.h"
 #include "Interfaces/TargetDeviceId.h"
 #include "Interfaces/ITargetPlatform.h"
 #include "Interfaces/ITargetPlatformManagerModule.h"
 #include "Interfaces/IProjectManager.h"
-#include "ISettingsModule.h"
 #include "SourceControlHelpers.h"
 #include "ISourceControlModule.h"
 #include "ITargetDeviceProxy.h"
@@ -34,13 +31,13 @@
 #include "DesktopPlatformModule.h"
 #include "PlatformInfo.h"
 #include "InstalledPlatformInfo.h"
-#include "Settings/EditorExperimentalSettings.h"
-#include "CookerSettings.h"
 #include "Framework/Commands/Commands.h"
 #include "Framework/Docking/TabManager.h"
 #include "Framework/MultiBox/MultiBoxBuilder.h"
 #include "Developer/DerivedDataCache/Public/DerivedDataCacheInterface.h"
 #include "Misc/MonitoredProcess.h"
+#include "EditorStyleSet.h"
+#include "CookerSettings.h"
 
 #include "Misc/App.h"
 #include "Framework/Application/SlateApplication.h"
@@ -173,11 +170,6 @@ protected:
 
 public:
 
-	/** Opens the Packaging settings tab */
-	static void PackagingSettings()
-	{
-		FModuleManager::LoadModuleChecked<ISettingsModule>("Settings").ShowViewer("Project", "Project", "Packaging");
-	}
 
 	static void OpenProjectLauncher()
 	{
@@ -280,11 +272,6 @@ public:
 		{
 			BuildCookRunParams += TEXT(" -installed");
 		}
-		int32 NumCookers = GetDefault<UEditorExperimentalSettings>()->MultiProcessCooking;
-		if (NumCookers > 0)
-		{
-			BuildCookRunParams += FString::Printf(TEXT(" -NumCookersToSpawn=%d"), NumCookers);
-		}
 
 
 
@@ -527,7 +514,7 @@ public:
 	void RegisterCommands()
 	{
 		UI_COMMAND(PackagingSettings, "Packaging Settings...", "Opens the settings for project packaging", EUserInterfaceActionType::Button, FInputChord());
-		ActionList->MapAction(PackagingSettings, FExecuteAction::CreateStatic(&FTurnkeySupportCallbacks::PackagingSettings));
+		ActionList->MapAction(PackagingSettings, FExecuteAction::CreateLambda([] {} ));
 
 	}
 
@@ -615,11 +602,11 @@ static void TurnkeyInstallSdk(FString PlatformName, bool bPreferFull, bool bForc
 			// update the Sdk status
 //			FDataDrivenPlatformInfoRegistry::UpdateSdkStatus();
 			GetTargetPlatformManager()->UpdateAfterSDKInstall(*PlatformName);
+#if WITH_ENGINE
 			RenderUtilsInit();
+#endif
 
-
-			// show restart dialog
-			FModuleManager::GetModuleChecked<ISettingsEditorModule>("SettingsEditor").OnApplicationRestartRequired();
+			FTurnkeyEditorSupport::ShowRestartToast();
 		});
 	}
 	);
@@ -923,15 +910,7 @@ static void LaunchOnDevice(const FString& DeviceId, const FString& DeviceName, b
 
 static void PrepareLaunchOn(FString DeviceId, FString DeviceName)
 {
-	ULevelEditorPlaySettings* PlaySettings = GetMutableDefault<ULevelEditorPlaySettings>();
-
-	PlaySettings->LastExecutedLaunchModeType = LaunchMode_OnDevice;
-	PlaySettings->LastExecutedLaunchDevice = DeviceId;
-	PlaySettings->LastExecutedLaunchName = DeviceName;
-
-	PlaySettings->PostEditChange();
-
-	PlaySettings->SaveConfig();
+	FTurnkeyEditorSupport::PrepareToLaunchRunningMap(DeviceId, DeviceName);
 }
 
 static void HandleLaunchOnDeviceActionExecute(FString DeviceId, FString DeviceName, bool bUseTurnkey)
@@ -1199,15 +1178,6 @@ TSharedRef<SWidget> FTurnkeySupportModule::MakeTurnkeyMenu() const
 			FUIAction(FExecuteAction::CreateStatic(&FTurnkeySupportCallbacks::OpenDeviceManager))
 		);
 
-		MenuBuilder.AddSeparator();
-
-		MenuBuilder.AddMenuEntry(
-			LOCTEXT("PackagingSettings", "Packaging Settings..."),
-			LOCTEXT("PackagingSettings_ToolTip", "Opens the settings for project packaging"),
-			FSlateIcon(),
-			FUIAction(FExecuteAction::CreateStatic(&FTurnkeySupportCallbacks::PackagingSettings))
-		);
-
 // 		MenuBuilder.AddMenuEntry(Commands.ZipUpProject,
 // 			NAME_None,
 // 			TAttribute<FText>(),
diff --git a/Engine/Source/Editor/TurnkeySupport/TurnkeySupport.Build.cs b/Engine/Source/Editor/TurnkeySupport/TurnkeySupport.Build.cs
--- a/Engine/Source/Editor/TurnkeySupport/TurnkeySupport.Build.cs
+++ b/Engine/Source/Editor/TurnkeySupport/TurnkeySupport.Build.cs
@@ -12,7 +12,7 @@ public class TurnkeySupport : ModuleRules
 				"Core",
 				"CoreUObject",
 				"ApplicationCore",
-				"Engine",
+//				"Engine",
 				"EngineSettings",
 				"InputCore",
 				"RHI",
@@ -23,10 +23,11 @@ public class TurnkeySupport : ModuleRules
 				"DesktopPlatform",
 				"WorkspaceMenuStructure",
 				"MessageLog",
-				"Projects",
-				"ToolMenus",
-				"LauncherServices",
-				"SettingsEditor"
+ 				"Projects",
+ 				"ToolMenus",
+ 				"LauncherServices",
+				"SourceControl",
+				"EditorStyle",
 			}
 		);
 
@@ -34,26 +35,25 @@ public class TurnkeySupport : ModuleRules
 		{
 			PrivateDependencyModuleNames.AddRange(
 				new string[] {
-					"UnrealEd",
-					"UATHelper",
-					"EditorStyle",
-					"SourceControl",
+				"UnrealEd",
+				"UATHelper",
+ 				"SettingsEditor"
+				}
+			);
+			PrivateIncludePathModuleNames.AddRange(
+				new string[] {
+					"AssetTools",
+					"DesktopPlatform",
+					"GameProjectGeneration",
+					"ProjectTargetPlatformEditor",
+					"LevelEditor",
+					"Settings",
+					"SourceCodeAccess",
+					"LocalizationDashboard",
 				}
 			);
 		}
 
-		PrivateIncludePathModuleNames.AddRange(
-			new string[] {
-				"AssetTools",
-				"DesktopPlatform",
-				"GameProjectGeneration",
-				"ProjectTargetPlatformEditor",
-				"LevelEditor",
-				"Settings",
-				"SourceCodeAccess",
-				"LocalizationDashboard",
-			}
-		);
 
 		PrivateIncludePaths.AddRange(
 			new string[] {
