- Fixed some non-unity issues

diff --git a/Engine/Source/Editor/TurnkeySupport/Private/TurnkeyEditorSupport.cpp b/Engine/Source/Editor/TurnkeySupport/Private/TurnkeyEditorSupport.cpp
--- a/Engine/Source/Editor/TurnkeySupport/Private/TurnkeyEditorSupport.cpp
+++ b/Engine/Source/Editor/TurnkeySupport/Private/TurnkeyEditorSupport.cpp
@@ -5,6 +5,7 @@
 
 #if WITH_EDITOR
 #include "UnrealEdMisc.h"
+#include "Framework/MultiBox/MultiBoxBuilder.h"
 #include "Editor/UnrealEdEngine.h"
 #include "UnrealEdGlobals.h"
 #include "PlatformInfo.h"
diff --git a/Engine/Source/Editor/TurnkeySupport/Private/TurnkeyEditorSupport.h b/Engine/Source/Editor/TurnkeySupport/Private/TurnkeyEditorSupport.h
--- a/Engine/Source/Editor/TurnkeySupport/Private/TurnkeyEditorSupport.h
+++ b/Engine/Source/Editor/TurnkeySupport/Private/TurnkeyEditorSupport.h
@@ -13,10 +13,10 @@ public:
 
 	static void PrepareToLaunchRunningMap(const FString& DeviceId, const FString& DeviceName);
 	static void LaunchRunningMap(const FString& DeviceId, const FString& DeviceName, bool bUseTurnkey);
-	static void AddEditorOptions(FMenuBuilder& MenuBuilder);
+	static void AddEditorOptions(class FMenuBuilder& MenuBuilder);
 
 	static bool DoesProjectHaveCode();
-	static void RunUAT(const FString& CommandLine, const FText& PlatformDisplayName, const FText& TaskName, const FText& TaskShortName, const FSlateBrush* TaskIcon, TFunction<void(FString, double)> ResultCallback=TFunction<void(FString, double)>());
+	static void RunUAT(const FString& CommandLine, const FText& PlatformDisplayName, const FText& TaskName, const FText& TaskShortName, const struct FSlateBrush* TaskIcon, TFunction<void(FString, double)> ResultCallback=TFunction<void(FString, double)>());
 
 	static bool ShowOKCancelDialog(FText Message, FText Title);
 	static void ShowRestartToast();
diff --git a/Engine/Source/Editor/TurnkeySupport/Private/TurnkeySupportModule.cpp b/Engine/Source/Editor/TurnkeySupport/Private/TurnkeySupportModule.cpp
--- a/Engine/Source/Editor/TurnkeySupport/Private/TurnkeySupportModule.cpp
+++ b/Engine/Source/Editor/TurnkeySupport/Private/TurnkeySupportModule.cpp
@@ -2,7 +2,6 @@
 
 #include "TurnkeySupportModule.h"
 #include "TurnkeySupport.h"
-#include "TurnkeyEditorSupport.h"
 
 #include "SlateOptMacros.h"
 #include "Widgets/Images/SImage.h"
@@ -38,12 +37,15 @@
 #include "Misc/MonitoredProcess.h"
 #include "EditorStyleSet.h"
 #include "CookerSettings.h"
+#include "TurnkeyEditorSupport.h"
 
 #include "Misc/App.h"
 #include "Framework/Application/SlateApplication.h"
 
 
-
+#if WITH_ENGINE
+#include "RenderUtils.h"
+#endif
 
 DEFINE_LOG_CATEGORY(LogTurnkeySupport);
 #define LOCTEXT_NAMESPACE "FTurnkeySupportModule"
diff --git a/Engine/Source/Editor/TurnkeySupport/Private/TurnkeySupportModule.h b/Engine/Source/Editor/TurnkeySupport/Private/TurnkeySupportModule.h
--- a/Engine/Source/Editor/TurnkeySupport/Private/TurnkeySupportModule.h
+++ b/Engine/Source/Editor/TurnkeySupport/Private/TurnkeySupportModule.h
@@ -17,7 +17,7 @@ public:
 	 *
 	 * @return	The newly-created menu widget
 	 */
-	virtual TSharedRef<SWidget> MakeTurnkeyMenu() const override;
+	virtual TSharedRef<class SWidget> MakeTurnkeyMenu() const override;
 
 	/**
 	 * @return	The newly-created menu widget
