- Made the Launch On menu be a normal drop down (not split button/combo box), named it Platforms menu
- Copied the quick launch platform items over to the Play menu, and now will repalce the Play button (which would change for Simulate/PIE/etc) to also do the quick launch on if that was the last one selected (using platforms large icon)
- Use Turnkey SDK installation availability info so that the Platforms menu can choose to show the proper SDK installation options that match what is supported by the studio
- Pass the -project path to Turnkey so a platform can check if a project has up-to-date prerequisites (like certain options need to be set in platform settings, etc)
- Use the platform icon to show warning/question to indicate status of SDK, with less text labels after the platform name in the Turnkey menu
- Added LoadSettingsForPlatform to ProjectPackgingSettings to be able to easily load packaging settings for another platform

diff --git a/Engine/Source/Developer/DeveloperToolSettings/Classes/Settings/ProjectPackagingSettings.h b/Engine/Source/Developer/DeveloperToolSettings/Classes/Settings/ProjectPackagingSettings.h
--- a/Engine/Source/Developer/DeveloperToolSettings/Classes/Settings/ProjectPackagingSettings.h
+++ b/Engine/Source/Developer/DeveloperToolSettings/Classes/Settings/ProjectPackagingSettings.h
@@ -91,6 +91,32 @@ enum class EProjectPackagingBlueprintNativizationMethod : uint8
 	Exclusive
 };
 
+USTRUCT()
+struct FProjectBuildSettings
+{
+	GENERATED_BODY();
+
+	/** The name for this custom build. It will be shown in menus for selection. */
+	UPROPERTY(EditAnywhere, Category="Packaging")
+	FString Name;
+
+	/** If this build step is specific to one or more platforms, add them here by name (note: use Windows, not Win64) */
+	UPROPERTY(EditAnywhere, Category="Packaging")
+	TArray<FString> SpecificPlatforms;
+
+	/**
+	 * The commandline for BuildCookRun UAT command. Some substitutions are allowed:
+	 *   {Project} - Replaced with the path to the project's .uproject file
+	 *   {Platform} - Replaced with the platform name this is run for
+	 *   {inivalue:Config:Section:Key} - Replaced with the value for Key in Config's Section. Ex: -archivedirectory={inivalue:Engine:CustomSettings:OverrideArchiveDir}
+	 *   {iniif:Token:Config:Section:Key} - Replaced with Token if the vlaue for Key in Config's Section evaluates to True. Ex: {iniif:-iostore:/Script/UnrealEd.ProjectPackagingSettings:bUseIoStore}
+	 * Because ProjectPackagingSettings is a common section to read, if Config:Section: are not specified for 'iniif' or 'inivalue', it will use the ProjectPackagingSettings settings:
+	 *   {iniif:-iostore:bUseIoStore}
+	 */
+	UPROPERTY(EditAnywhere, Category="Packaging")
+	FString BuildCookRunParams;
+};
+
 /**
  * Implements the Editor's user settings.
  */
@@ -100,6 +126,7 @@ class DEVELOPERTOOLSETTINGS_API UProjectPackagingSettings
 {
 	GENERATED_UCLASS_BODY()
 
+	
 	/**
 	 * This class was moved from UnrealEd module, but to allow it to be used by developer tools, like UFE, it has moved to this module.
 	 * However, for back-compat, we want to use the old name in the ini files, so that everything works without needing to touch every
@@ -469,6 +496,11 @@ public:
 	UPROPERTY(config, EditAnywhere, Category=Packaging, AdvancedDisplay, meta=(DisplayName="Additional Non-Asset Directories To Copy for dedicated server only", RelativeToGameContentDir))
 	TArray<FDirectoryPath> DirectoriesToAlwaysStageAsNonUFSServer;	
 
+	/**
+	 * A list of custom builds that will show up in the Platforms menu to allow customized builds that make sense for your project. Will show up near Package Project in the Platforms menu.
+	 */
+	UPROPERTY(config, EditAnywhere, Category=Packaging, meta=(DisplayName = "Additional builds for this project."))
+	TArray<FProjectBuildSettings> ExtraProjectBuilds;
 
 
 	/**
@@ -489,9 +521,12 @@ private:
 	UPROPERTY(config)
 	TMap<FName, FName> PerPlatformTargetPlatformName;
 
-		/** Helper array used to mirror Blueprint asset selections across edits */
+	/** Helper array used to mirror Blueprint asset selections across edits */
 	TArray<FFilePath> CachedNativizeBlueprintAssets;
 	
+	/** THe platform to LoadConfig for to get platform-specific packaging settings */
+	FString ConfigPlatform;
+
 public:
 
 	// UObject Interface
@@ -520,6 +555,23 @@ public:
 	/** Gets the current build target, checking that it's valid, and the default build target if it is not */
 	const FTargetInfo* GetBuildTargetInfo() const;
 
+	/** For non-default object instances, this will LoadConfig for a specific platform,*/
+	void LoadSettingsForPlatform(FString PlatformName)
+	{
+		ConfigPlatform = PlatformName;
+		LoadConfig();
+	}
+
+	FString GetConfigPlatform() const
+	{
+		return ConfigPlatform;
+	}
+
+	virtual const TCHAR* GetConfigOverridePlatform() const override
+	{
+		return ConfigPlatform.Len() ? *ConfigPlatform : Super::GetConfigOverridePlatform();
+	}
+
 private:
 	/** Fix up cooking paths after they've been edited or laoded */
 	void FixCookingPaths();
diff --git a/Engine/Source/Editor/TurnkeySupport/Private/TurnkeyEditorSupport.cpp b/Engine/Source/Editor/TurnkeySupport/Private/TurnkeyEditorSupport.cpp
--- a/Engine/Source/Editor/TurnkeySupport/Private/TurnkeyEditorSupport.cpp
+++ b/Engine/Source/Editor/TurnkeySupport/Private/TurnkeyEditorSupport.cpp
@@ -19,7 +19,8 @@
 #include "ISettingsModule.h"
 #include "ISettingsEditorModule.h"
 #include "Settings/EditorExperimentalSettings.h"
-
+#include "Interfaces/IMainFrameModule.h"
+#include "ToolMenus.h"
 #endif
 
 #define LOCTEXT_NAMESPACE "FTurnkeyEditorSupport"
@@ -42,19 +43,25 @@ FString FTurnkeyEditorSupport::GetUATOptions()
 }
 
 
-void FTurnkeyEditorSupport::AddEditorOptions(FMenuBuilder& MenuBuilder)
+void FTurnkeyEditorSupport::AddEditorOptions(FToolMenuSection& Section)
 {
 #if WITH_EDITOR
-	MenuBuilder.AddSeparator();
+	Section.AddSeparator(NAME_None);
 
-	MenuBuilder.AddMenuEntry(
+	Section.AddMenuEntry(
+		NAME_None,
 		LOCTEXT("OpenPackagingSettings", "Packaging Settings..."),
 		LOCTEXT("OpenPackagingSettings_ToolTip", "Opens the settings for project packaging."),
 		FSlateIcon(FEditorStyle::GetStyleSetName(), "DeviceDetails.TabIcon"),
 		FUIAction(FExecuteAction::CreateLambda([] { FModuleManager::LoadModuleChecked<ISettingsModule>("Settings").ShowViewer("Project", "Project", "Packaging"); }))
 	);
 
-	FModuleManager::LoadModuleChecked<IProjectTargetPlatformEditorModule>("ProjectTargetPlatformEditor").AddOpenProjectTargetPlatformEditorMenuItem(MenuBuilder);
+	// use AddDynamicEntry to be able to get a MenuBuilder for external code that needs one
+	Section.AddDynamicEntry(NAME_None, FNewToolMenuDelegateLegacy::CreateLambda([](FMenuBuilder& MenuBuilder, UToolMenu* Menu)
+		{
+			FModuleManager::LoadModuleChecked<IProjectTargetPlatformEditorModule>("ProjectTargetPlatformEditor").AddOpenProjectTargetPlatformEditorMenuItem(MenuBuilder);
+		})
+	);
 #endif
 }
 
@@ -206,5 +213,16 @@ bool FTurnkeyEditorSupport::CheckSupportedPlatforms(FName IniPlatformName)
 	return true;
 }
 
+void FTurnkeyEditorSupport::ShowInstallationHelp(FName IniPlatformName, FString DocLink)
+{
+#if WITH_EDITOR
+	// broadcast this, and assume someone will pick it up
+	IMainFrameModule& MainFrameModule = FModuleManager::GetModuleChecked<IMainFrameModule>(TEXT("MainFrame"));
+	MainFrameModule.BroadcastMainFrameSDKNotInstalled(IniPlatformName.ToString(), DocLink);
+#else
+	unimplemented();
+#endif
+}
+
 
 #undef LOCTEXT_NAMESPACE
diff --git a/Engine/Source/Editor/TurnkeySupport/Private/TurnkeyEditorSupport.h b/Engine/Source/Editor/TurnkeySupport/Private/TurnkeyEditorSupport.h
--- a/Engine/Source/Editor/TurnkeySupport/Private/TurnkeyEditorSupport.h
+++ b/Engine/Source/Editor/TurnkeySupport/Private/TurnkeyEditorSupport.h
@@ -13,7 +13,7 @@ public:
 
 	static void PrepareToLaunchRunningMap(const FString& DeviceId, const FString& DeviceName);
 	static void LaunchRunningMap(const FString& DeviceId, const FString& DeviceName, bool bUseTurnkey);
-	static void AddEditorOptions(class FMenuBuilder& MenuBuilder);
+	static void AddEditorOptions(struct FToolMenuSection& MenuBuilder);
 
 	static bool DoesProjectHaveCode();
 	static void RunUAT(const FString& CommandLine, const FText& PlatformDisplayName, const FText& TaskName, const FText& TaskShortName, const struct FSlateBrush* TaskIcon, TFunction<void(FString, double)> ResultCallback=TFunction<void(FString, double)>());
@@ -21,5 +21,6 @@ public:
 	static bool ShowOKCancelDialog(FText Message, FText Title);
 	static void ShowRestartToast();
 	static bool CheckSupportedPlatforms(FName IniPlatformName);
+	static void ShowInstallationHelp(FName IniPlatformName, FString DocLink);
 
 };
diff --git a/Engine/Source/Editor/TurnkeySupport/Private/TurnkeySupportModule.cpp b/Engine/Source/Editor/TurnkeySupport/Private/TurnkeySupportModule.cpp
--- a/Engine/Source/Editor/TurnkeySupport/Private/TurnkeySupportModule.cpp
+++ b/Engine/Source/Editor/TurnkeySupport/Private/TurnkeySupportModule.cpp
@@ -17,6 +17,7 @@
 #include "Misc/DataDrivenPlatformInfoRegistry.h"
 #include "Async/Async.h"
 #include "Misc/FileHelper.h"
+#include "Misc/ConfigCacheIni.h"
 #include "Interfaces/TargetDeviceId.h"
 #include "Interfaces/ITargetPlatform.h"
 #include "Interfaces/ITargetPlatformManagerModule.h"
@@ -37,6 +38,8 @@
 #include "Misc/MonitoredProcess.h"
 #include "EditorStyleSet.h"
 #include "CookerSettings.h"
+#include "UObject/UObjectIterator.h"
+#include "ToolMenus.h"
 #include "TurnkeyEditorSupport.h"
 
 #include "Misc/App.h"
@@ -168,7 +171,6 @@ protected:
 		return bBuild;
 	}
 
-
 public:
 
 
@@ -189,13 +191,34 @@ public:
 			return false;
 		}
 
-		// PrepForDebugging needs the platform to specify how
-		if (Mode == EPrepareContentMode::PrepareForDebugging)
+// 		// PrepForDebugging needs the platform to specify how
+// 		if (Mode == EPrepareContentMode::PrepareForDebugging)
+// 		{
+// 			return FDataDrivenPlatformInfoRegistry::GetPlatformInfo(IniPlatformName).PrepareForDebuggingOptions != TEXT("");
+// 		}
+
+		return true;
+	}
+
+	static UProjectPackagingSettings* GetPackagingSettingsForPlatform(FName IniPlatformName)
+	{
+		FString PlatformString = IniPlatformName.ToString();
+		UProjectPackagingSettings* PackagingSettings = nullptr;
+		for (TObjectIterator<UProjectPackagingSettings> Itr; Itr; ++Itr)
+		{
+			if (Itr->GetConfigPlatform() == PlatformString)
+			{
+				PackagingSettings = *Itr;
+				break;
+			}
+		}
+		if (PackagingSettings == nullptr)
 		{
-			return FDataDrivenPlatformInfoRegistry::GetPlatformInfo(IniPlatformName).PrepareForDebuggingOptions != TEXT("");
+			PackagingSettings = NewObject<UProjectPackagingSettings>(GetTransientPackage());
+			PackagingSettings->LoadSettingsForPlatform(PlatformString);
 		}
 
-		return true;
+		return PackagingSettings;
 	}
 
 	static void CookOrPackage(FName IniPlatformName, EPrepareContentMode Mode)
@@ -208,7 +231,7 @@ public:
 
 		// get all the helper objects
 		const FString UBTPlatformString = PlatformInfo->DataDrivenPlatformInfo->UBTPlatformString;
-		UProjectPackagingSettings* PackagingSettings = Cast<UProjectPackagingSettings>(UProjectPackagingSettings::StaticClass()->GetDefaultObject());
+		UProjectPackagingSettings* PackagingSettings = GetPackagingSettingsForPlatform(IniPlatformName);
 		const FString ProjectPath = FPaths::IsProjectFilePathSet() ? FPaths::ConvertRelativePathToFull(FPaths::GetProjectFilePath()) : FPaths::RootDir() / FApp::GetProjectName() / FApp::GetProjectName() + TEXT(".uproject");
 
 
@@ -386,16 +409,16 @@ public:
 				BuildCookRunParams += FString::Printf(TEXT(" -clientconfig=%s"), LexToString(ConfigurationInfo.Configuration));
 			}
 		}
-		else if (Mode == EPrepareContentMode::PrepareForDebugging)
-		{
-			const ITargetPlatform* TargetPlatform = GetTargetPlatformManager()->FindTargetPlatform(PlatformInfo->Name);
-			if (ShouldBuildProject(PackagingSettings, TargetPlatform))
-			{
-				BuildCookRunParams += TEXT(" -build");
-			}
-
-			BuildCookRunParams += FString::Printf(TEXT(" %s"), *FDataDrivenPlatformInfoRegistry::GetPlatformInfo(IniPlatformName).PrepareForDebuggingOptions);
-		}
+// 		else if (Mode == EPrepareContentMode::PrepareForDebugging)
+// 		{
+// 			const ITargetPlatform* TargetPlatform = GetTargetPlatformManager()->FindTargetPlatform(PlatformInfo->Name);
+// 			if (ShouldBuildProject(PackagingSettings, TargetPlatform))
+// 			{
+// 				BuildCookRunParams += TEXT(" -build");
+// 			}
+// 
+// 			BuildCookRunParams += FString::Printf(TEXT(" %s"), *FDataDrivenPlatformInfoRegistry::GetPlatformInfo(IniPlatformName).PrepareForDebuggingOptions);
+// 		}
 		else if (Mode == EPrepareContentMode::CookOnly)
 		{
 			ContentPrepDescription = LOCTEXT("CookingContentTaskName", "Cooking content");
@@ -413,7 +436,7 @@ public:
 		}
 
 
-		FString TurnkeyParams = FString::Printf(TEXT(" -command=VerifySdk -platform=%s -UpdateIfNeeded -EditorIO"), *UBTPlatformString);
+		FString TurnkeyParams = FString::Printf(TEXT("-command=VerifySdk -platform=%s -UpdateIfNeeded -EditorIO -project=\"%s\""), *UBTPlatformString, *ProjectPath);
 
 		FString CommandLine = FString::Printf(TEXT("-ScriptsForProject=\"%s\" Turnkey %s BuildCookRun %s"),
 			*ProjectPath,
@@ -424,6 +447,137 @@ public:
 		FTurnkeyEditorSupport::RunUAT(CommandLine, PlatformInfo->DisplayName, ContentPrepDescription, ContentPrepTaskName, ContentPrepIcon);
 	}
 
+	static bool CanExecuteCustomBuild(FName IniPlatformName, FProjectBuildSettings Build)
+	{
+		if (GetTargetPlatformManager()->FindTargetPlatform(IniPlatformName.ToString()) == nullptr)
+		{
+			return false;
+		}
+
+		return true;
+	}
+
+	static FString GetIniSetting(FString Spec)
+	{
+		// handle these cases:
+		//  iniif:-option:Engine:/Script/Module.Class:bUseOption
+		//  iniif:-option:bUseOption   [convenience for ProjectPackagingSettings setting]
+		//  inivalue:Engine:/Script/Module.Class:SomeSetting
+		//  inivalue:SomeSetting       [convenience for ProjectPackagingSettings setting]
+
+		TArray<FString> Tokens;
+		Spec.ParseIntoArray(Tokens, TEXT(":"));
+
+		FString ConfigName;
+		FString SectionName;
+		FString Key;
+		FString IniIfValue;
+		if (Tokens[0] == TEXT("iniif"))
+		{
+			if (Tokens.Num() == 3)
+			{
+				ConfigName = TEXT("Game");
+				SectionName = TEXT("/Script/UnrealEd.ProjectPackagingSettings");
+				Key = Tokens[2];
+				IniIfValue = Tokens[1];
+			}
+			else if (Tokens.Num() == 5)
+			{
+				ConfigName = Tokens[2];
+				SectionName = Tokens[3];
+				Key = Tokens[4];
+				IniIfValue = Tokens[1];
+			}
+			else
+			{
+				UE_LOG(LogTurnkeySupport, Warning, TEXT("Found a bad iniif spec: %s"), *Spec);
+				return TEXT("");
+			}
+		}
+		else if (Tokens[0] == TEXT("inivalue"))
+		{
+			if (Tokens.Num() == 2)
+			{
+				ConfigName = TEXT("Game");
+				SectionName = TEXT("/Script/UnrealEd.ProjectPackagingSettings");
+				Key = Tokens[1];
+			}
+			else if (Tokens.Num() == 4)
+			{
+				ConfigName = Tokens[1];
+				SectionName = Tokens[2];
+				Key = Tokens[3];
+			}
+			else
+			{
+				UE_LOG(LogTurnkeySupport, Warning, TEXT("Found a bad inivalue spec: %s"), *Spec);
+				return TEXT("");
+			}
+		}
+		else
+		{
+			UE_LOG(LogTurnkeySupport, Warning, TEXT("Found a bad ini spec: %s"), *Spec);
+			return TEXT("");
+		}
+
+		// get the value, if it exists (or empty string if not)
+		FString FoundValue;
+		GConfig->GetString(*SectionName, *Key, FoundValue, ConfigName);
+
+		if (Tokens[0] == TEXT("iniif"))
+		{
+			if (FCString::ToBool(*FoundValue))
+			{
+				return IniIfValue;
+			}
+			return TEXT("");
+		}
+
+		return FoundValue;
+
+	}
+
+	// @todo since Turnkey can execute builds with ExecuteBuild, we can probably just let it do all of the work so that we don't have the commandline fixup code in two places
+	static void ExecuteCustomBuild(FName IniPlatformName, FProjectBuildSettings Build)
+	{
+		const PlatformInfo::FTargetPlatformInfo* PlatformInfo = PlatformInfo::FindPlatformInfo(GetDefault<UProjectPackagingSettings>()->GetTargetPlatformForPlatform(IniPlatformName));
+
+		FString CommandLine = Build.BuildCookRunParams;
+
+		const FString ProjectPath = FPaths::IsProjectFilePathSet() ? FPaths::ConvertRelativePathToFull(FPaths::GetProjectFilePath()) : FPaths::RootDir() / FApp::GetProjectName() / FApp::GetProjectName() + TEXT(".uproject");
+		CommandLine = CommandLine.Replace(TEXT("{Project}"), *ProjectPath);
+
+		CommandLine = CommandLine.Replace(TEXT("{Platform}"), *IniPlatformName.ToString());
+
+		int Offset = 0;
+		while (Offset != -1)
+		{
+			Offset = CommandLine.Find(TEXT("{ini"), ESearchCase::IgnoreCase, ESearchDir::FromStart, Offset);
+			if (Offset != -1)
+			{
+				// skip over the {
+				Offset++;
+
+				// look for the close brace, don't need to worry about nested braces
+				int CloseBrace = CommandLine.Find(TEXT("}"), ESearchCase::IgnoreCase, ESearchDir::FromStart, Offset);
+				if (CloseBrace == -1)
+				{
+					Offset = -1;
+				}
+				else
+				{
+					FString IniSpec = CommandLine.Mid(Offset, CloseBrace - Offset);
+					FString ReplacedValue = GetIniSetting(IniSpec);
+					CommandLine = CommandLine.Mid(0, Offset - 1) + ReplacedValue + CommandLine.Mid(CloseBrace + 1);
+				}
+			}
+		}
+
+
+
+		FTurnkeyEditorSupport::RunUAT(CommandLine, PlatformInfo->DisplayName, LOCTEXT("Turnkey_CustomTaskName", "Executing Custom Build"), LOCTEXT("Turnkey_CustomTaskName", "Custom"), FEditorStyle::GetBrush(TEXT("MainFrame.PackageProject")));
+	}
+
 	static void PackageBuildConfiguration(const PlatformInfo::FTargetPlatformInfo* Info, EProjectPackagingBuildConfigurations BuildConfiguration)
 	{
 		UProjectPackagingSettings* PackagingSettings = GetMutableDefault<UProjectPackagingSettings>();
@@ -541,6 +695,19 @@ TSharedRef<FUICommandList> FTurnkeySupportCommands::ActionList = MakeShareable(n
 
 
 
+static void ShowInstallationHelp(FName IniPlatformName)
+{
+// 	const ITargetPlatform* Platform = GetTargetPlatformManager()->FindTargetPlatform(IniPlatformName.ToString());
+// 	bool bProjectHasCode;
+// 	EBuildConfiguration Configuration;
+// 	bool bRequiresAssetNativization;
+// 	FString TutorialPath, DocumentationPath;
+// 	FText LogMessage;
+// 
+// 	Platform->CheckRequirements(bProjectHasCode, Configuration, bRequiresAssetNativization, TutorialPath, DocumentationPath, LogMessage);
+
+	FTurnkeyEditorSupport::ShowInstallationHelp(IniPlatformName, FDataDrivenPlatformInfoRegistry::GetPlatformInfo(IniPlatformName).SDKTutorial);
+}
 
 static void TurnkeyInstallSdk(FString PlatformName, bool bPreferFull, bool bForceInstall, FString DeviceId)
 {
@@ -615,26 +782,6 @@ static void TurnkeyInstallSdk(FString PlatformName, bool bPreferFull, bool bForc
 
 static TAttribute<FText> MakeSdkStatusAttribute(FName IniPlatformName, TSharedPtr< ITargetDeviceProxy> DeviceProxy)
 {
-	// 	FFormatNamedArguments LabelArguments;
-	// 	LabelArguments.Add(TEXT("DeviceName"), FText::FromString(DeviceProxy->GetName()));
-	// 
-	// 	if (!DeviceProxy->IsConnected())
-	// 	{
-	// 		LabelArguments.Add(TEXT("HostUser"), LOCTEXT("DisconnectedHint", " [Disconnected]"));
-	// 	}
-	// 	else if (DeviceProxy->GetHostUser() != FPlatformProcess::UserName(false))
-	// 	{
-	// 		LabelArguments.Add(TEXT("HostUser"), FText::FromString(DeviceProxy->GetHostUser()));
-	// 	}
-	// 	else
-	// 	{
-	// 		LabelArguments.Add(TEXT("HostUser"), FText::GetEmpty());
-	// 	}
-	// 
-	// 	OutLabel = FText::Format(LOCTEXT("LaunchDeviceLabel", "{DeviceName}{HostUser}"), LabelArguments);
-	// 	TAttribute<FText> = Make
-
-
 	FString DisplayString = DeviceProxy ? DeviceProxy->GetName() : IniPlatformName.ToString();
 	FString DeviceId = DeviceProxy ? DeviceProxy->GetTargetDeviceId(NAME_None) : FString();
 
@@ -643,19 +790,47 @@ static TAttribute<FText> MakeSdkStatusAttribute(FName IniPlatformName, TSharedPt
 		// get the status, or Unknown if it's not there
 		ETurnkeyPlatformSdkStatus Status = DeviceId.Len() ? ITurnkeySupportModule::Get().GetSdkInfoForDeviceId(DeviceId).Status : ITurnkeySupportModule::Get().GetSdkInfo(IniPlatformName, false).Status;
 
-		// @todo turnkey: Have premade FText's by SdkStatus for speed
-		const TCHAR* Desc =
-			Status == ETurnkeyPlatformSdkStatus::Querying ? TEXT("Querying...") :
-			Status == ETurnkeyPlatformSdkStatus::Valid ? TEXT("Valid Sdk") :
-			Status == ETurnkeyPlatformSdkStatus::OutOfDate ? TEXT("Outdated Sdk") :
-			Status == ETurnkeyPlatformSdkStatus::NoSdk ? TEXT("No Sdk") :
-			Status == ETurnkeyPlatformSdkStatus::FlashValid ? TEXT("Valid Flash") :
-			Status == ETurnkeyPlatformSdkStatus::FlashOutOfDate ? TEXT("Outdated Flash") :
-			TEXT("???");
-		return FText::FromString(FString::Printf(TEXT("%s (%s)"), *DisplayString, Desc));
+		if (Status == ETurnkeyPlatformSdkStatus::Querying)
+		{
+			FFormatNamedArguments LabelArguments;
+			LabelArguments.Add(TEXT("DisplayName"), FText::FromString(DisplayString));
+			return FText::Format(LOCTEXT("SDKStatusLabel", "{DisplayName} (Querying...)"), LabelArguments);
+		}
+		return FText::FromString(DisplayString);
 	}));
 }
 
+//static TAttribute<FSlateIcon> MakePlatformSdkIconAttribute(FName IniPlatformName, TSharedPtr< ITargetDeviceProxy> DeviceProxy)
+static FSlateIcon MakePlatformSdkIconAttribute(FName IniPlatformName, TSharedPtr< ITargetDeviceProxy> DeviceProxy)
+{
+	FString DeviceId = DeviceProxy ? DeviceProxy->GetTargetDeviceId(NAME_None) : FString();
+
+//	return TAttribute<FSlateIcon>::Create(TAttribute<FSlateIcon>::FGetter::CreateLambda([IniPlatformName, DeviceId]()
+		{
+			// get the status, or Unknown if it's not there
+			ETurnkeyPlatformSdkStatus Status = DeviceId.Len() ? ITurnkeySupportModule::Get().GetSdkInfoForDeviceId(DeviceId).Status : ITurnkeySupportModule::Get().GetSdkInfo(IniPlatformName, false).Status;
+
+			if (Status == ETurnkeyPlatformSdkStatus::OutOfDate || Status == ETurnkeyPlatformSdkStatus::NoSdk || Status == ETurnkeyPlatformSdkStatus::FlashOutOfDate)
+			{
+				return FSlateIcon(FEditorStyle::GetStyleSetName(), TEXT("Icons.Warning"));
+			}
+			else if (Status == ETurnkeyPlatformSdkStatus::Error)
+			{
+				return FSlateIcon(FEditorStyle::GetStyleSetName(), TEXT("Icons.Error"));
+			}
+			else if (Status == ETurnkeyPlatformSdkStatus::Unknown)
+			{
+				return FSlateIcon(FEditorStyle::GetStyleSetName(), TEXT("Icons.Help"));
+			}
+			else
+			{
+				return FSlateIcon(FEditorStyle::GetStyleSetName(), FDataDrivenPlatformInfoRegistry::GetPlatformInfo(IniPlatformName).GetIconStyleName(EPlatformIconSize::Normal));
+			}
+		}
+//		));
+}
+
+
 
 
 
@@ -692,16 +867,36 @@ static void MakeTurnkeyPlatformMenu(FMenuBuilder& MenuBuilder, FName IniPlatform
 				FCanExecuteAction::CreateStatic(&FTurnkeySupportCallbacks::CanCookOrPackage, IniPlatformName, EPrepareContentMode::CookOnly)
 			)
 		);
+// 
+// 		MenuBuilder.AddMenuEntry(
+// 			LOCTEXT("Turnkey_PrepareForDebugging", "Prepare For Debugging"),
+// 			LOCTEXT("TurnkeyTooltip_PrepareForDebugging", "Prepare this project for debugging"),
+// 			FSlateIcon(),
+// 			FUIAction(
+// 				FExecuteAction::CreateStatic(&FTurnkeySupportCallbacks::CookOrPackage, IniPlatformName, EPrepareContentMode::PrepareForDebugging),
+// 				FCanExecuteAction::CreateStatic(&FTurnkeySupportCallbacks::CanCookOrPackage, IniPlatformName, EPrepareContentMode::PrepareForDebugging)
+// 			)
+// 		);
+// 
 
-		MenuBuilder.AddMenuEntry(
-			LOCTEXT("Turnkey_PrepareForDebugging", "Prepare For Debugging"),
-			LOCTEXT("TurnkeyTooltip_PrepareForDebugging", "Prepare this project for debugging"),
-			FSlateIcon(),
-			FUIAction(
-				FExecuteAction::CreateStatic(&FTurnkeySupportCallbacks::CookOrPackage, IniPlatformName, EPrepareContentMode::PrepareForDebugging),
-				FCanExecuteAction::CreateStatic(&FTurnkeySupportCallbacks::CanCookOrPackage, IniPlatformName, EPrepareContentMode::PrepareForDebugging)
-			)
-		);
+		FString PlatformString = IniPlatformName.ToString();
+		UProjectPackagingSettings* PackagingSettings = FTurnkeySupportCallbacks::GetPackagingSettingsForPlatform(IniPlatformName);
+
+		for (FProjectBuildSettings Build : PackagingSettings->ExtraProjectBuilds)
+		{
+			if (Build.SpecificPlatforms.Num() == 0 || Build.SpecificPlatforms.Contains(PlatformString))
+			{
+				MenuBuilder.AddMenuEntry(
+					FText::FromString(Build.Name),
+					LOCTEXT("TurnkeyTooltip_CustomBuild", "Execute a custom build (this comes from Packaging Settings)"),
+					FSlateIcon(),
+					FUIAction(
+						FExecuteAction::CreateStatic(&FTurnkeySupportCallbacks::ExecuteCustomBuild, IniPlatformName, Build),
+						FCanExecuteAction::CreateStatic(&FTurnkeySupportCallbacks::CanExecuteCustomBuild, IniPlatformName, Build)
+					)
+				);
+			}
+		}
 
 		MenuBuilder.EndSection();
 
@@ -790,7 +985,8 @@ static void MakeTurnkeyPlatformMenu(FMenuBuilder& MenuBuilder, FName IniPlatform
 						FSlateIcon(),
 						FExecuteAction::CreateStatic(TurnkeyInstallSdk, UBTPlatformString, true, false, DeviceId)
 					);
-				})
+				}),
+				false, MakePlatformSdkIconAttribute(IniPlatformName, Proxy)
 			);
 		}
 
@@ -810,52 +1006,68 @@ static void MakeTurnkeyPlatformMenu(FMenuBuilder& MenuBuilder, FName IniPlatform
 	);
 
 	FString NoDevice;
-	if (SdkInfo.Status == ETurnkeyPlatformSdkStatus::OutOfDate)
+	if (SdkInfo.bCanInstallFullSdk || SdkInfo.bCanInstallAutoSdk)
 	{
-		MenuBuilder.AddMenuEntry(
-			LOCTEXT("Turnkey_InstallSdkMinimal", "Update Sdk (Prefer Minimal)"),
-			LOCTEXT("TurnkeyTooltip_InstallSdkMinimal", "Attempt to update an Sdk, as hosted by your studio. Will attempt to install a minimal Sdk (useful for building/running only)"),
-			FSlateIcon(),
-			FExecuteAction::CreateStatic(TurnkeyInstallSdk, UBTPlatformString, false, false, NoDevice)
-		);
+		if (SdkInfo.Status == ETurnkeyPlatformSdkStatus::OutOfDate)
+		{
+			MenuBuilder.AddMenuEntry(
+				LOCTEXT("Turnkey_InstallSdkMinimal", "Update Sdk"),
+				LOCTEXT("TurnkeyTooltip_InstallSdkMinimal", "Attempt to update an Sdk, as hosted by your studio. Will attempt to install a minimal Sdk (useful for building/running only)"),
+				FSlateIcon(),
+				FExecuteAction::CreateStatic(TurnkeyInstallSdk, UBTPlatformString, false, false, NoDevice)
+			);
 
-		MenuBuilder.AddMenuEntry(
-			LOCTEXT("Turnkey_InstallSdkFull", "Update Sdk (Prefer Full)"),
-			LOCTEXT("TurnkeyTooltip_InstallSdkMinimal", "Attempt to update an Sdk, as hosted by your studio. Will attempt to install a full Sdk (useful profiling or other use cases)"),
-			FSlateIcon(),
-			FExecuteAction::CreateStatic(TurnkeyInstallSdk, UBTPlatformString, true, false, NoDevice)
-		);
-	}
-	else if (SdkInfo.Status == ETurnkeyPlatformSdkStatus::Valid)
-	{
-		MenuBuilder.AddMenuEntry(
-			LOCTEXT("Turnkey_InstallSdkMinimal", "Force Reinstall Sdk (Prefer Minimal)"),
-			LOCTEXT("TurnkeyTooltip_InstallSdkMinimal", "Attempt to force re-install an Sdk, as hosted by your studio. Will attempt to install a minimal Sdk (useful for building/running only)"),
-			FSlateIcon(),
-			FExecuteAction::CreateStatic(TurnkeyInstallSdk, UBTPlatformString, false, true, NoDevice)
-		);
+			if (SdkInfo.bCanInstallFullSdk && SdkInfo.bCanInstallAutoSdk)
+			{
+				MenuBuilder.AddMenuEntry(
+					LOCTEXT("Turnkey_InstallSdkFull", "Update Sdk (Prefer Full)"),
+					LOCTEXT("TurnkeyTooltip_InstallSdkMinimal", "Attempt to update an Sdk, as hosted by your studio. Will attempt to install a full Sdk (useful profiling or other use cases)"),
+					FSlateIcon(),
+					FExecuteAction::CreateStatic(TurnkeyInstallSdk, UBTPlatformString, true, false, NoDevice)
+				);
+			}
+		}
+		else if (SdkInfo.Status == ETurnkeyPlatformSdkStatus::Valid)
+		{
+			MenuBuilder.AddMenuEntry(
+				LOCTEXT("Turnkey_InstallSdkMinimal", "Force Reinstall Sdk (Prefer Minimal)"),
+				LOCTEXT("TurnkeyTooltip_InstallSdkMinimal", "Attempt to force re-install an Sdk, as hosted by your studio. Will attempt to install a minimal Sdk (useful for building/running only)"),
+				FSlateIcon(),
+				FExecuteAction::CreateStatic(TurnkeyInstallSdk, UBTPlatformString, false, true, NoDevice)
+			);
 
-		MenuBuilder.AddMenuEntry(
-			LOCTEXT("Turnkey_InstallSdkFull", "Force Reinstall (Prefer Full)"),
-			LOCTEXT("TurnkeyTooltip_InstallSdkMinimal", "Attempt to force re-install an Sdk, as hosted by your studio. Will attempt to install a full Sdk (useful profiling or other use cases)"),
-			FSlateIcon(),
-			FExecuteAction::CreateStatic(TurnkeyInstallSdk, UBTPlatformString, true, true, NoDevice)
-		);
+			MenuBuilder.AddMenuEntry(
+				LOCTEXT("Turnkey_InstallSdkFull", "Force Reinstall (Prefer Full)"),
+				LOCTEXT("TurnkeyTooltip_InstallSdkMinimal", "Attempt to force re-install an Sdk, as hosted by your studio. Will attempt to install a full Sdk (useful profiling or other use cases)"),
+				FSlateIcon(),
+				FExecuteAction::CreateStatic(TurnkeyInstallSdk, UBTPlatformString, true, true, NoDevice)
+			);
+		}
+		else
+		{
+			MenuBuilder.AddMenuEntry(
+				LOCTEXT("Turnkey_InstallSdkMinimal", "Install Sdk (Prefer Minimal)"),
+				LOCTEXT("TurnkeyTooltip_InstallSdkMinimal", "Attempt to install an Sdk, as hosted by your studio. Will attempt to install a minimal Sdk (useful for building/running only)"),
+				FSlateIcon(),
+				FExecuteAction::CreateStatic(TurnkeyInstallSdk, UBTPlatformString, false, false, NoDevice)
+			);
+
+			MenuBuilder.AddMenuEntry(
+				LOCTEXT("Turnkey_InstallSdkFull", "Install Sdk (Prefer Full)"),
+				LOCTEXT("TurnkeyTooltip_InstallSdkMinimal", "Attempt to install an Sdk, as hosted by your studio. Will attempt to install a full Sdk (useful profiling or other use cases)"),
+				FSlateIcon(),
+				FExecuteAction::CreateStatic(TurnkeyInstallSdk, UBTPlatformString, true, false, NoDevice)
+			);
+		}
 	}
 	else
 	{
+		// if Turnkey can't be used for this platform, then show old-school documentation
 		MenuBuilder.AddMenuEntry(
-			LOCTEXT("Turnkey_InstallSdkMinimal", "Install Sdk (Prefer Minimal)"),
-			LOCTEXT("TurnkeyTooltip_InstallSdkMinimal", "Attempt to install an Sdk, as hosted by your studio. Will attempt to install a minimal Sdk (useful for building/running only)"),
-			FSlateIcon(),
-			FExecuteAction::CreateStatic(TurnkeyInstallSdk, UBTPlatformString, false, false, NoDevice)
-		);
-
-		MenuBuilder.AddMenuEntry(
-			LOCTEXT("Turnkey_InstallSdkFull", "Install Sdk (Prefer Full)"),
-			LOCTEXT("TurnkeyTooltip_InstallSdkMinimal", "Attempt to install an Sdk, as hosted by your studio. Will attempt to install a full Sdk (useful profiling or other use cases)"),
+			LOCTEXT("Turnkey_ShowDocumentation", "Installation Help..."),
+			LOCTEXT("TurnkeyTooltip_ShowDocumentation", "Show documentation with help installing the SDK for this platform"),
 			FSlateIcon(),
-			FExecuteAction::CreateStatic(TurnkeyInstallSdk, UBTPlatformString, true, false, NoDevice)
+			FExecuteAction::CreateStatic(ShowInstallationHelp, IniPlatformName)
 		);
 	}
 }
@@ -926,7 +1138,7 @@ static bool HandleLaunchOnDeviceActionCanExecute(FString DeviceName)
 	return CanLaunchOnDevice(DeviceName);
 }
 
-static void GenerateDeviceProxyMenuParams(TSharedPtr<ITargetDeviceProxy> DeviceProxy, FName PlatformName, FUIAction& OutAction, FText& OutTooltip)
+static void GenerateDeviceProxyMenuParams(TSharedPtr<ITargetDeviceProxy> DeviceProxy, FName PlatformName, FUIAction& OutAction, FText& OutTooltip, FOnQuickLaunchSelected ExternalOnClickDelegate)
 {
 	// 	// create an All_<platform>_devices_on_<host> submenu
 	// 	if (DeviceProxy->IsAggregated())
@@ -944,11 +1156,15 @@ static void GenerateDeviceProxyMenuParams(TSharedPtr<ITargetDeviceProxy> DeviceP
 	// 	}
 
 		// ... create an action...
-	FString DeviceId = DeviceProxy->GetTargetDeviceId(NAME_None);
 	OutAction = FUIAction(
-		FExecuteAction::CreateStatic(&HandleLaunchOnDeviceActionExecute, DeviceId, DeviceProxy->GetName(), true)
-		//		, FCanExecuteAction::CreateStatic(&HandleLaunchOnDeviceActionCanExecute, DeviceProxy->GetName())
-	);
+		FExecuteAction::CreateLambda([DeviceProxy, ExternalOnClickDelegate]()
+			{
+				FString DeviceId = DeviceProxy->GetTargetDeviceId(NAME_None);
+				HandleLaunchOnDeviceActionExecute(DeviceId, DeviceProxy->GetName(), true);
+				ExternalOnClickDelegate.ExecuteIfBound(DeviceId);
+			}
+	//		, FCanExecuteAction::CreateStatic(&HandleLaunchOnDeviceActionCanExecute, DeviceProxy->GetName())
+	));
 
 	// ... generate tooltip text
 	FFormatNamedArguments TooltipArguments;
@@ -966,34 +1182,17 @@ static void GenerateDeviceProxyMenuParams(TSharedPtr<ITargetDeviceProxy> DeviceP
 		FText TooltipLine2 = FText::Format(LOCTEXT("LaunchDevicePlatformWarning", "{DisplayName} is not listed as a target platform for this project, so may not run as expected."), TooltipArguments);
 		OutTooltip = FText::Format(FText::FromString(TEXT("{0}\n\n{1}")), OutTooltip, TooltipLine2);
 	}
-
 }
 
 
-
-
-
-
-
-
-
-
-TSharedRef<SWidget> FTurnkeySupportModule::MakeTurnkeyMenu() const
+void FTurnkeySupportModule::MakeQuickLaunchItems(class UToolMenu* Menu, FOnQuickLaunchSelected ExternalOnClickDelegate) const
 {
-	FTurnkeySupportCommands::Register();
-
-	const FTurnkeySupportCommands& Commands = FTurnkeySupportCommands::Get();
-
-	const bool bShouldCloseWindowAfterMenuSelection = true;
-	FMenuBuilder MenuBuilder(bShouldCloseWindowAfterMenuSelection, FTurnkeySupportCommands::ActionList);
-
-	// shared devices section
-	ITargetDeviceServicesModule* TargetDeviceServicesModule = static_cast<ITargetDeviceServicesModule*>(FModuleManager::Get().LoadModule(TEXT("TargetDeviceServices")));
-
-	TArray<FString> DeviceIdsToQuery;
+	FToolMenuSection& MenuSection = Menu->AddSection("QuickLaunchDevices", LOCTEXT("QuickLaunch", "Quick Launch"));
 
-	MenuBuilder.BeginSection("LevelEditorLaunchDevices", LOCTEXT("TurnkeySection_LaunchButtonDevices", "Quick Launch"));
+	MenuSection.AddDynamicEntry("PlatformsMenu", FNewToolMenuSectionDelegate::CreateLambda([ExternalOnClickDelegate](FToolMenuSection& DynamicSection)
 	{
+		TArray<FString> DeviceIdsToQuery;
+		ITargetDeviceServicesModule* TargetDeviceServicesModule = static_cast<ITargetDeviceServicesModule*>(FModuleManager::Get().LoadModule(TEXT("TargetDeviceServices")));
 		for (const auto& Pair : FDataDrivenPlatformInfoRegistry::GetAllPlatformInfos())
 		{
 			if (Pair.Value.bIsFakePlatform)
@@ -1017,50 +1216,50 @@ TSharedRef<SWidget> FTurnkeySupportModule::MakeTurnkeyMenu() const
 				// always use the first one, after sorting
 				FUIAction Action;
 				FText Tooltip;
-				GenerateDeviceProxyMenuParams(DeviceProxies[0], PlatformName, Action, Tooltip);
+				GenerateDeviceProxyMenuParams(DeviceProxies[0], PlatformName, Action, Tooltip, ExternalOnClickDelegate);
 
 				if (DeviceProxies.Num() == 1)
 				{
-					MenuBuilder.AddMenuEntry(
+					UE_LOG(LogTurnkeySupport, Display, TEXT("Adding device menu item for %s"), *DeviceProxies[0]->GetName());
+					DynamicSection.AddMenuEntry(
+						NAME_None,
 						MakeSdkStatusAttribute(PlatformName, DeviceProxies[0]),
 						Tooltip,
-						FSlateIcon(FEditorStyle::GetStyleSetName(), Pair.Value.GetIconStyleName(EPlatformIconSize::Normal)),
-						Action,
-						NAME_None,
-						EUserInterfaceActionType::Button
+						MakePlatformSdkIconAttribute(PlatformName, DeviceProxies[0]),
+						Action
 					);
 				}
 				else
 				{
-					MenuBuilder.AddSubMenu(
+					DynamicSection.AddSubMenu(
+						NAME_None,
 						MakeSdkStatusAttribute(PlatformName, DeviceProxies[0]),
 						Tooltip,
-						FNewMenuDelegate::CreateLambda([TargetDeviceServicesModule, PlatformName, IconStyle = Pair.Value.GetIconStyleName(EPlatformIconSize::Normal)](FMenuBuilder& SubMenuBuilder) 
-						{
-							// re-get the proxies, just in case they changed
-							TArray<TSharedPtr<ITargetDeviceProxy>> DeviceProxies;
-							TargetDeviceServicesModule->GetDeviceProxyManager()->GetAllProxies(PlatformName, DeviceProxies);
-							// for each one, put an entry (even the one that was in the outer menu, for less confusion)
-							for (const TSharedPtr<ITargetDeviceProxy>& Proxy : DeviceProxies)
+						FNewMenuDelegate::CreateLambda([TargetDeviceServicesModule, PlatformName, &ExternalOnClickDelegate](FMenuBuilder& SubMenuBuilder)
 							{
-								FUIAction SubAction;
-								FText SubTooltip;
-								GenerateDeviceProxyMenuParams(Proxy, PlatformName, SubAction, SubTooltip);
-								SubMenuBuilder.AddMenuEntry(
-									MakeSdkStatusAttribute(PlatformName, Proxy),
-									SubTooltip,
-									FSlateIcon(FEditorStyle::GetStyleSetName(), IconStyle),
-									SubAction,
-									NAME_None,
-									EUserInterfaceActionType::Button
-								);
-							}
-						}),
+								// re-get the proxies, just in case they changed
+								TArray<TSharedPtr<ITargetDeviceProxy>> DeviceProxies;
+								TargetDeviceServicesModule->GetDeviceProxyManager()->GetAllProxies(PlatformName, DeviceProxies);
+								// for each one, put an entry (even the one that was in the outer menu, for less confusion)
+								for (const TSharedPtr<ITargetDeviceProxy>& Proxy : DeviceProxies)
+								{
+									FUIAction SubAction;
+									FText SubTooltip;
+									GenerateDeviceProxyMenuParams(Proxy, PlatformName, SubAction, SubTooltip, ExternalOnClickDelegate);
+									SubMenuBuilder.AddMenuEntry(
+										MakeSdkStatusAttribute(PlatformName, Proxy),
+										SubTooltip,
+										MakePlatformSdkIconAttribute(PlatformName, Proxy),
+										SubAction,
+										NAME_None,
+										EUserInterfaceActionType::Button
+									);
+								}
+							}),
 						Action,
-						NAME_None,
 						EUserInterfaceActionType::Check,
 						false,
-						FSlateIcon(FEditorStyle::GetStyleSetName(), Pair.Value.GetIconStyleName(EPlatformIconSize::Normal)),
+						MakePlatformSdkIconAttribute(PlatformName, nullptr),
 						true
 						);
 				}
@@ -1077,151 +1276,211 @@ TSharedRef<SWidget> FTurnkeySupportModule::MakeTurnkeyMenu() const
 				}
 			}
 		}
-	}
-	MenuBuilder.EndSection();
-
-
-	MenuBuilder.BeginSection("CookerSettings");
-
-	MenuBuilder.AddMenuEntry(
-		LOCTEXT("CookOnTheFlyOnLaunch", "Enable cooking on the fly"),
-		LOCTEXT("CookOnTheFlyOnLaunchDescription", "Cook on the fly instead of cooking upfront when launching"),
-		FSlateIcon(),
-		FUIAction(
-			FExecuteAction::CreateStatic(&FTurnkeySupportCallbacks::SetCookOnTheFly),
-			FCanExecuteAction::CreateStatic(&FTurnkeySupportCallbacks::CanSetCookOnTheFly),
-			FIsActionChecked::CreateStatic(&FTurnkeySupportCallbacks::SetCookOnTheFlyIsChecked)
-		),
-		NAME_None,
-		EUserInterfaceActionType::ToggleButton
-	);
-
-	{
-		MenuBuilder.AddWidget(
-			SNew(STextBlock)
-			.ColorAndOpacity(FSlateColor::UseSubduedForeground())
-			.Text(LOCTEXT("ZoomToFitHorizontal", "Launching a game on a different device will change your default 'Launch' device in the toolbar"))
-			.WrapTextAt(300),
-			FText::GetEmpty()
-		);
-	}
-	MenuBuilder.EndSection();
 
-	MenuBuilder.BeginSection("AllPlatforms", LOCTEXT("TurnkeyMenu_ManagePlatforms", "Content/Sdk/Device Management"));
-	TMap<FName, const FDataDrivenPlatformInfo*> UncompiledPlatforms;
-	for (const auto& Pair : FDataDrivenPlatformInfoRegistry::GetAllPlatformInfos())
-	{
-		if (Pair.Value.bIsFakePlatform || Pair.Value.bEnabledForUse == false)
+		// if we don't have an external delegate to call, then this is the internally included items in the Platforms menu and we can add the extra option(s)
+		if (!ExternalOnClickDelegate.IsBound())
 		{
-			continue;
+			DynamicSection.AddMenuEntry(
+				NAME_None,
+				LOCTEXT("CookOnTheFlyOnLaunch", "Enable cooking on the fly"),
+				LOCTEXT("CookOnTheFlyOnLaunchDescription", "Cook on the fly instead of cooking upfront when launching"),
+				FSlateIcon(),
+				FUIAction(
+					FExecuteAction::CreateStatic(&FTurnkeySupportCallbacks::SetCookOnTheFly),
+					FCanExecuteAction::CreateStatic(&FTurnkeySupportCallbacks::CanSetCookOnTheFly),
+					FIsActionChecked::CreateStatic(&FTurnkeySupportCallbacks::SetCookOnTheFlyIsChecked)
+				),
+				EUserInterfaceActionType::ToggleButton
+			);
 		}
 
-		FName PlatformName = Pair.Key;
-		const FDataDrivenPlatformInfo& Info = Pair.Value;
 
-		if (!FDataDrivenPlatformInfoRegistry::HasCompiledSupportForPlatform(PlatformName, FDataDrivenPlatformInfoRegistry::EPlatformNameType::Ini))
+		// now kick-off any devices that need to be updated
+		if (DeviceIdsToQuery.Num() > 0)
 		{
-			UncompiledPlatforms.Add(PlatformName, &Info);
-			continue;
+			ITurnkeySupportModule::Get().UpdateSdkInfoForDevices(DeviceIdsToQuery);
 		}
-
-		MenuBuilder.AddSubMenu(
-			MakeSdkStatusAttribute(PlatformName, nullptr),
-			FText::FromString(PlatformName.ToString()),
-			FNewMenuDelegate::CreateStatic(&MakeTurnkeyPlatformMenu, PlatformName, TargetDeviceServicesModule),
-			false,
-			FSlateIcon(FEditorStyle::GetStyleSetName(), Pair.Value.GetIconStyleName(EPlatformIconSize::Normal)),
-			true
-		);
 	}
+	));
+}
 
-	if (UncompiledPlatforms.Num() != 0)
-	{
-		MenuBuilder.AddSeparator(NAME_None);
-
-		MenuBuilder.AddSubMenu(
-			LOCTEXT("Turnkey_UncompiledPlatforms", "Platforms With No Compiled Support"),
-			LOCTEXT("Turnkey_UncompiledPlatformsToolTip", "List of platforms that you have access to, but support is not compiled in to the editor. It may be caused by missing an SDK, so you attempt to install an SDK here."),
-			FNewMenuDelegate::CreateLambda([UncompiledPlatforms, TargetDeviceServicesModule](FMenuBuilder& SubMenuBuilder)
-			{
-				for (const auto It : UncompiledPlatforms)
-				{
-					SubMenuBuilder.AddSubMenu(
-						MakeSdkStatusAttribute(It.Key, nullptr),
-						FText::FromString(It.Key.ToString()),
-						FNewMenuDelegate::CreateStatic(&MakeTurnkeyPlatformMenu, It.Key, TargetDeviceServicesModule),
-						false,
-						FSlateIcon(FEditorStyle::GetStyleSetName(), It.Value->GetIconStyleName(EPlatformIconSize::Normal)),
-						true
-					);
-				}
-			})
-		);
-	}
+TSharedRef<SWidget> FTurnkeySupportModule::MakeTurnkeyMenuWidget() const
+{
+	FTurnkeySupportCommands::Register();
+	const FTurnkeySupportCommands& Commands = FTurnkeySupportCommands::Get();
 
-	MenuBuilder.EndSection();
+	const bool bShouldCloseWindowAfterMenuSelection = true;
+//	FMenuBuilder MenuBuilder(bShouldCloseWindowAfterMenuSelection, FTurnkeySupportCommands::ActionList);
 
+	static const FName MenuName("UnrealEd.PlayWorldCommands.PlatformsMenu");
 
-	// options section
-	MenuBuilder.BeginSection("TurnkeyOptions", LOCTEXT("TurnkeySection_Options", "Options and Settings"));
+	if (!UToolMenus::Get()->IsMenuRegistered(MenuName))
 	{
-		MenuBuilder.AddMenuEntry(
-			LOCTEXT("OpenProjectLauncher", "Project Launcher..."),
-			LOCTEXT("OpenProjectLauncher_ToolTip", "Open the Project Launcher for advanced packaging, deploying and launching of your projects"),
-			FSlateIcon(FEditorStyle::GetStyleSetName(), "Launcher.TabIcon"),
-			FUIAction(FExecuteAction::CreateStatic(&FTurnkeySupportCallbacks::OpenProjectLauncher))
-		);
-
-		MenuBuilder.AddMenuEntry(
-			LOCTEXT("OpenDeviceManager", "Device Manager..."),
-			LOCTEXT("OpenDeviceManager_ToolTip", "View and manage connected devices."),
-			FSlateIcon(FEditorStyle::GetStyleSetName(), "DeviceDetails.TabIcon"),
-			FUIAction(FExecuteAction::CreateStatic(&FTurnkeySupportCallbacks::OpenDeviceManager))
-		);
-
-// 		MenuBuilder.AddMenuEntry(Commands.ZipUpProject,
-// 			NAME_None,
-// 			TAttribute<FText>(),
-// 			TAttribute<FText>(),
-// 			FSlateIcon()
-// 		);
-// 		MenuBuilder.AddMenuEntry(Commands.PackagingSettings,
-// 			NAME_None,
-// 			TAttribute<FText>(),
-// 			TAttribute<FText>(),
-// 			FSlateIcon()
-// 		);
+		UToolMenu* Menu = UToolMenus::Get()->RegisterMenu(MenuName);
 
-		FTurnkeyEditorSupport::AddEditorOptions(MenuBuilder);
-	}
-	MenuBuilder.EndSection();
+		FOnQuickLaunchSelected EmptyFunc;
+		MakeQuickLaunchItems(Menu, EmptyFunc);
 
-	// now kick-off any devices that need to be updated
-	if (DeviceIdsToQuery.Num() > 0)
-	{
-		ITurnkeySupportModule::Get().UpdateSdkInfoForDevices(DeviceIdsToQuery);
-	}
+		// need to make this dyamic so icons, etc can update with SDK 
+		// shared devices section
 
-	return MenuBuilder.MakeWidget();
-}
+		FToolMenuSection& ManagePlatformsSection = Menu->AddSection("AllPlatforms", LOCTEXT("TurnkeyMenu_ManagePlatforms", "Content/Sdk/Device Management"));
+		ManagePlatformsSection.AddDynamicEntry(NAME_None, FNewToolMenuSectionDelegate::CreateLambda([](FToolMenuSection& PlatformsSection)
+		{
+			ITargetDeviceServicesModule* TargetDeviceServicesModule = static_cast<ITargetDeviceServicesModule*>(FModuleManager::Get().LoadModule(TEXT("TargetDeviceServices")));
 
+			TMap<FName, const FDataDrivenPlatformInfo*> UncompiledPlatforms;
+			TMap<FName, const FDataDrivenPlatformInfo*> UnsupportedPlatforms;
 
+			FProjectStatus ProjectStatus;
+			bool bProjectStatusIsValid = IProjectManager::Get().QueryStatusForCurrentProject(ProjectStatus);
 
+			for (const auto& Pair : FDataDrivenPlatformInfoRegistry::GetAllPlatformInfos())
+			{
+				if (Pair.Value.bIsFakePlatform || Pair.Value.bEnabledForUse == false)
+				{
+					continue;
+				}
 
+				FName PlatformName = Pair.Key;
+				const FDataDrivenPlatformInfo& Info = Pair.Value;
 
+				if (!FDataDrivenPlatformInfoRegistry::HasCompiledSupportForPlatform(PlatformName, FDataDrivenPlatformInfoRegistry::EPlatformNameType::Ini))
+				{	
+					UncompiledPlatforms.Add(PlatformName, &Info);
+					continue;
+				}
 
-// some shared functionality
-static void PrepForTurnkeyReport(FString& Command, FString& BaseCommandline, FString& ReportFilename)
-{
-	static int ReportIndex = 0;
+				if (bProjectStatusIsValid && !ProjectStatus.IsTargetPlatformSupported(PlatformName))
+				{
+					UnsupportedPlatforms.Add(PlatformName, &Info);
+					continue;
+				}
 
-	FString LogFilename = FPaths::ConvertRelativePathToFull(FPaths::Combine(FPaths::ProjectIntermediateDir(), *FString::Printf(TEXT("TurnkeyLog_%d.log"), ReportIndex)));
-	ReportFilename = FPaths::ConvertRelativePathToFull(FPaths::Combine(FPaths::ProjectIntermediateDir(), *FString::Printf(TEXT("TurnkeyReport_%d.log"), ReportIndex++)));
+				PlatformsSection.AddSubMenu(
+					NAME_None,
+					MakeSdkStatusAttribute(PlatformName, nullptr),
+					FText::FromString(PlatformName.ToString()),
+					FNewMenuDelegate::CreateStatic(&MakeTurnkeyPlatformMenu, PlatformName, TargetDeviceServicesModule),
+					false,
+					MakePlatformSdkIconAttribute(PlatformName, nullptr),
+					true
+				);
+			}
 
-	// make sure intermediate directory exists
-	IFileManager::Get().MakeDirectory(*FPaths::ProjectIntermediateDir());
+			if (UnsupportedPlatforms.Num() != 0)
+			{
+				PlatformsSection.AddSeparator(NAME_None);
 
-	Command = TEXT("{EngineDir}Build/BatchFiles/RunuAT");
+				PlatformsSection.AddSubMenu(
+					NAME_None,
+					LOCTEXT("Turnkey_UnsupportedPlatforms", "Platforms Not Supported by Project"),
+					LOCTEXT("Turnkey_UnsupportedPlatformsToolTip", "List of platforms that are not marked as supported by this platform. Use the \"Supported Platforms...\""),
+					FNewMenuDelegate::CreateLambda([UnsupportedPlatforms, TargetDeviceServicesModule](FMenuBuilder& SubMenuBuilder)
+						{
+							for (const auto It : UnsupportedPlatforms)
+							{
+								SubMenuBuilder.AddSubMenu(
+									MakeSdkStatusAttribute(It.Key, nullptr),
+									FText::FromString(It.Key.ToString()),
+									FNewMenuDelegate::CreateStatic(&MakeTurnkeyPlatformMenu, It.Key, TargetDeviceServicesModule),
+									false,
+									MakePlatformSdkIconAttribute(It.Key, nullptr),
+									true
+								);
+							}
+						})
+				);
+			}
+
+			if (UncompiledPlatforms.Num() != 0)
+			{
+				PlatformsSection.AddSeparator(NAME_None);
+
+				PlatformsSection.AddSubMenu(
+					NAME_None,
+					LOCTEXT("Turnkey_UncompiledPlatforms", "Platforms With No Compiled Support"),
+					LOCTEXT("Turnkey_UncompiledPlatformsToolTip", "List of platforms that you have access to, but support is not compiled in to the editor. It may be caused by missing an SDK, so you attempt to install an SDK here."),
+					FNewMenuDelegate::CreateLambda([UncompiledPlatforms, TargetDeviceServicesModule](FMenuBuilder& SubMenuBuilder)
+						{
+							for (const auto It : UncompiledPlatforms)
+							{
+								SubMenuBuilder.AddSubMenu(
+									MakeSdkStatusAttribute(It.Key, nullptr),
+									FText::FromString(It.Key.ToString()),
+									FNewMenuDelegate::CreateStatic(&MakeTurnkeyPlatformMenu, It.Key, TargetDeviceServicesModule),
+									false,
+									MakePlatformSdkIconAttribute(It.Key, nullptr),
+									true
+								);
+							}
+						})
+				);
+			}
+		}));
+
+
+		// options section
+		FToolMenuSection& OptionsSection = Menu->AddSection("TurnkeyOptions", LOCTEXT("TurnkeySection_Options", "Options and Settings"));
+		{
+			OptionsSection.AddMenuEntry(
+				NAME_None,
+				LOCTEXT("OpenProjectLauncher", "Project Launcher..."),
+				LOCTEXT("OpenProjectLauncher_ToolTip", "Open the Project Launcher for advanced packaging, deploying and launching of your projects"),
+				FSlateIcon(FEditorStyle::GetStyleSetName(), "Launcher.TabIcon"),
+				FUIAction(FExecuteAction::CreateStatic(&FTurnkeySupportCallbacks::OpenProjectLauncher))
+			);
+
+			OptionsSection.AddMenuEntry(
+				NAME_None,
+				LOCTEXT("OpenDeviceManager", "Device Manager..."),
+				LOCTEXT("OpenDeviceManager_ToolTip", "View and manage connected devices."),
+				FSlateIcon(FEditorStyle::GetStyleSetName(), "DeviceDetails.TabIcon"),
+				FUIAction(FExecuteAction::CreateStatic(&FTurnkeySupportCallbacks::OpenDeviceManager))
+			);
+
+			FTurnkeyEditorSupport::AddEditorOptions(OptionsSection);
+		}
+	}
+
+	FToolMenuContext MenuContext(FTurnkeySupportCommands::ActionList);
+	return UToolMenus::Get()->GenerateWidget(MenuName, MenuContext);
+
+}
+
+void FTurnkeySupportModule::MakeTurnkeyMenu(FToolMenuSection& MenuSection) const
+{
+	// make sure the DeviceProxyManager is going _before_ we create the menu contents dynamically, so that devices will show up
+	ITargetDeviceServicesModule* TargetDeviceServicesModule = static_cast<ITargetDeviceServicesModule*>(FModuleManager::Get().LoadModule(TEXT("TargetDeviceServices")));
+	TargetDeviceServicesModule->GetDeviceProxyManager();
+
+ 	MenuSection.AddEntry(FToolMenuEntry::InitComboButton(
+		"PlatformsMenu",
+		FUIAction(),
+		FOnGetContent::CreateLambda([this] { return MakeTurnkeyMenuWidget(); }),
+		LOCTEXT("PlatformMenu", "Platforms"),
+		LOCTEXT("PlatformMenu_Tooltip", "Platform related actions and settings (Launching, Packaging, custom builds, etc)"),
+		FSlateIcon(FEditorStyle::GetStyleSetName(), "PlayWorld.RepeatLastLaunch"), // not great name for a good "platforms" icon
+		false,
+		"PlatformsMenu"
+	));
+}
+
+
+
+
+// some shared functionality
+static void PrepForTurnkeyReport(FString& Command, FString& BaseCommandline, FString& ReportFilename)
+{
+	static int ReportIndex = 0;
+
+	FString LogFilename = FPaths::ConvertRelativePathToFull(FPaths::Combine(FPaths::ProjectIntermediateDir(), *FString::Printf(TEXT("TurnkeyLog_%d.log"), ReportIndex)));
+	ReportFilename = FPaths::ConvertRelativePathToFull(FPaths::Combine(FPaths::ProjectIntermediateDir(), *FString::Printf(TEXT("TurnkeyReport_%d.log"), ReportIndex++)));
+
+	// make sure intermediate directory exists
+	IFileManager::Get().MakeDirectory(*FPaths::ProjectIntermediateDir());
+
+	Command = TEXT("{EngineDir}Build/BatchFiles/RunuAT");
 	//	Command = TEXT("{EngineDir}/Binaries/DotNET/AutomationTool.exe");
 	BaseCommandline = FString::Printf(TEXT("Turnkey -utf8output -WaitForUATMutex -command=VerifySdk -ReportFilename=\"%s\" -log=\"%s\""), *ReportFilename, *LogFilename);
 
@@ -1279,6 +1538,8 @@ bool GetSdkInfoFromTurnkey(FString Line, FName& PlatformName, FString& DeviceId,
 			SdkInfo.Status = ETurnkeyPlatformSdkStatus::NoSdk;
 		}
 	}
+	SdkInfo.bCanInstallFullSdk = FlagsString.Contains(TEXT("Support_FullSdk"));
+	SdkInfo.bCanInstallAutoSdk = FlagsString.Contains(TEXT("Support_AutoSdk"));
 
 	return true;
 }
@@ -1355,7 +1616,7 @@ void FTurnkeySupportModule::UpdateSdkInfo()
 				{
 					for (FString& Line : Contents)
 					{
-						UE_LOG(LogTurnkeySupport, Log, TEXT("Turnkey Platform: %s"), *Line);
+ 						UE_LOG(LogTurnkeySupport, Log, TEXT("Turnkey Platform: %s"), *Line);
 
 						// parse a Turnkey line
 						FName PlatformName;
@@ -1368,9 +1629,10 @@ void FTurnkeySupportModule::UpdateSdkInfo()
 
 						// we received a platform from UAT that we don't know about in the editor. this can happen if you have a UBT/UAT that was compiled with platform access
 						// but then you are running without that platform synced. skip this platform and move on
-						if (PerPlatformSdkInfo.Contains(PlatformName))
+						if (!PerPlatformSdkInfo.Contains(PlatformName))
 						{
 							UE_LOG(LogTurnkeySupport, Log, TEXT("Received platform %s from Turnkey, but the engine doesn't know about it. Skipping..."), *PlatformName.ToString());
+							continue;
 						}
 
 						// check if we had already set a ManualSDK - and don't set it again. Because of the way AutoSDKs are activated in the editor after the first call to Turnkey,
@@ -1389,8 +1651,8 @@ void FTurnkeySupportModule::UpdateSdkInfo()
 						}
 
 
-						UE_LOG(LogTurnkeySupport, Log, TEXT("[TEST] Turnkey Platform: %s - %d, Installed: %s, AudoSDK: %s, Allowed: %s-%s"), *PlatformName.ToString(), (int)SdkInfo.Status, *SdkInfo.InstalledVersion,
-							*SdkInfo.AutoSDKVersion, *SdkInfo.MinAllowedVersion, *SdkInfo.MaxAllowedVersion);
+// 						UE_LOG(LogTurnkeySupport, Log, TEXT("[TEST] Turnkey Platform: %s - %d, Installed: %s, AudoSDK: %s, Allowed: %s-%s"), *PlatformName.ToString(), (int)SdkInfo.Status, *SdkInfo.InstalledVersion,
+// 							*SdkInfo.AutoSDKVersion, *SdkInfo.MinAllowedVersion, *SdkInfo.MaxAllowedVersion);
 					}
 				}
 			}
@@ -1521,6 +1783,27 @@ void FTurnkeySupportModule::UpdateSdkInfoForDevices(TArray<FString> PlatformDevi
 	TurnkeyProcess->Launch();
 }
 
+/**
+ * Runs Turnkey to get the Sdk information for all known platforms
+ */
+void FTurnkeySupportModule::RepeatQuickLaunch(FString DeviceId)
+{
+	UE_LOG(LogTurnkeySupport, Display, TEXT("Launching on %s"), *DeviceId);
+
+	ITargetDeviceServicesModule* TargetDeviceServicesModule = static_cast<ITargetDeviceServicesModule*>(FModuleManager::Get().LoadModule(TEXT("TargetDeviceServices")));
+	TSharedPtr<ITargetDeviceProxy> Proxy = TargetDeviceServicesModule->GetDeviceProxyManager()->FindProxyDeviceForTargetDevice(DeviceId);
+
+	if (Proxy.IsValid())
+	{
+		HandleLaunchOnDeviceActionExecute(DeviceId, Proxy->GetName(), true);
+	}
+	else
+	{
+		// @todo show error toast
+	}
+}
+
+
 FTurnkeySdkInfo FTurnkeySupportModule::GetSdkInfo(FName PlatformName, bool bBlockIfQuerying) const
 {
 	FScopeLock Lock(&GTurnkeySection);
@@ -1555,6 +1838,8 @@ void FTurnkeySupportModule::ClearDeviceStatus(FName PlatformName)
 
 void FTurnkeySupportModule::StartupModule( )
 {
+	
+
 }
 
 
@@ -1744,11 +2029,371 @@ void FInternalPlayWorldCommandCallbacks::HandleShowSDKTutorial(FString PlatformN
 
 
 
+#endif
+
 
+// another round of junk, some duplicates with above
+#if 0
 
-#endif
 
 
+/////// Will move these into TurnkeySupportModule soon:
+
+
+
+bool FInternalPlayWorldCommandCallbacks::CanLaunchOnDevice(const FString& DeviceName)
+{
+	if (!GUnrealEd->IsPlayingViaLauncher())
+	{
+		static TWeakPtr<ITargetDeviceProxyManager> DeviceProxyManagerPtr;
+
+		if (!DeviceProxyManagerPtr.IsValid())
+		{
+			ITargetDeviceServicesModule* TargetDeviceServicesModule = FModuleManager::Get().LoadModulePtr<ITargetDeviceServicesModule>(TEXT("TargetDeviceServices"));
+			if (TargetDeviceServicesModule)
+			{
+				DeviceProxyManagerPtr = TargetDeviceServicesModule->GetDeviceProxyManager();
+			}
+		}
+
+		TSharedPtr<ITargetDeviceProxyManager> DeviceProxyManager = DeviceProxyManagerPtr.Pin();
+		if (DeviceProxyManager.IsValid())
+		{
+			TSharedPtr<ITargetDeviceProxy> DeviceProxy = DeviceProxyManager->FindProxy(DeviceName);
+			if (DeviceProxy.IsValid() && DeviceProxy->IsConnected() && DeviceProxy->IsAuthorized())
+			{
+				return true;
+			}
+
+			// check if this is an aggregate proxy
+			TArray<TSharedPtr<ITargetDeviceProxy>> Devices;
+			DeviceProxyManager->GetProxies(FName(*DeviceName), false, Devices);
+
+			// returns true if the game can be launched al least on 1 device
+			for (auto DevicesIt = Devices.CreateIterator(); DevicesIt; ++DevicesIt)
+			{
+				TSharedPtr<ITargetDeviceProxy> DeviceAggregateProxy = *DevicesIt;
+				if (DeviceAggregateProxy.IsValid() && DeviceAggregateProxy->IsConnected() && DeviceAggregateProxy->IsAuthorized())
+				{
+					return true;
+				}
+			}
+
+		}
+	}
+
+	return false;
+}
+
+
+void FInternalPlayWorldCommandCallbacks::LaunchOnDevice(const FString& DeviceId, const FString& DeviceName, bool bUseTurnkey)
+{
+	FTargetDeviceId TargetDeviceId;
+	if (FTargetDeviceId::Parse(DeviceId, TargetDeviceId))
+	{
+		const PlatformInfo::FTargetPlatformInfo* const PlatformInfo = PlatformInfo::FindPlatformInfo(*TargetDeviceId.GetPlatformName());
+		FString UBTPlatformName = PlatformInfo->DataDrivenPlatformInfo->UBTPlatformString;
+		FString IniPlatformName = PlatformInfo->IniPlatformName.ToString();
+
+		check(PlatformInfo);
+
+		if (FInstalledPlatformInfo::Get().IsPlatformMissingRequiredFile(UBTPlatformName))
+		{
+			if (!FInstalledPlatformInfo::OpenInstallerOptions())
+			{
+				FMessageDialog::Open(EAppMsgType::Ok, LOCTEXT("MissingPlatformFilesLaunch", "Missing required files to launch on this platform."));
+			}
+			return;
+		}
+
+		if (FModuleManager::LoadModuleChecked<IProjectTargetPlatformEditorModule>("ProjectTargetPlatformEditor").ShowUnsupportedTargetWarning(*TargetDeviceId.GetPlatformName()))
+		{
+			GUnrealEd->CancelPlayingViaLauncher();
+
+			FRequestPlaySessionParams::FLauncherDeviceInfo DeviceInfo;
+			DeviceInfo.DeviceId = DeviceId;
+			DeviceInfo.DeviceName = DeviceName;
+			// @todo turnkey: we set this to false because we will kick off a Turnkey run before cooking, etc, to get an early warning. however, if it's too difficult
+			// to get an error back from CreateUatTask, then we should set this to bUseTurnkey and remove the block below, and let the code in FLauncherWorker::CreateAndExecuteTasks handle it
+			DeviceInfo.bUpdateDeviceFlash = false;
+
+			FRequestPlaySessionParams SessionParams;
+			SessionParams.SessionDestination = EPlaySessionDestinationType::Launcher;
+			SessionParams.LauncherTargetDevice = DeviceInfo;
+
+			// if we want to check device flash before we start cooking, kick it off now. we could delay this 
+			if (bUseTurnkey)
+			{
+				FString CommandLine = FString::Printf(TEXT("Turnkey -command=VerifySdk -UpdateIfNeeded -platform=%s -EditorIO -noturnkeyvariables -device=%s -utf8output -WaitForUATMutex"), *UBTPlatformName, *DeviceName);
+				FText TaskName = LOCTEXT("VerifyingSDK", "Verifying SDK and Device");
+
+				IUATHelperModule::Get().CreateUatTask(CommandLine, FText::FromString(IniPlatformName), TaskName, TaskName, FEditorStyle::GetBrush(TEXT("MainFrame.PackageProject")),
+					[SessionParams](FString Result, double)
+					{
+						// unfortunate string comparison for success
+						bool bWasSuccessful = Result == TEXT("Completed");
+						AsyncTask(ENamedThreads::GameThread, [SessionParams, bWasSuccessful]()
+							{
+								if (bWasSuccessful)
+								{
+									GUnrealEd->RequestPlaySession(SessionParams);
+								}
+								else
+								{
+									TSharedRef<SWindow> Win = OpenMsgDlgInt_NonModal(EAppMsgType::YesNo, LOCTEXT("SDKCheckFailed", "SDK Verification failed. Would you like to attempt the Launch On anyway?"), LOCTEXT("SDKCheckFailedTitle", "SDK Verification"),
+										FOnMsgDlgResult::CreateLambda([SessionParams](const TSharedRef<SWindow>&, EAppReturnType::Type Choice)
+											{
+												if (Choice == EAppReturnType::Yes)
+												{
+													GUnrealEd->RequestPlaySession(SessionParams);
+												}
+											}));
+									Win->ShowWindow();
+								}
+							});
+					});
+			}
+			else
+			{
+				GUnrealEd->RequestPlaySession(SessionParams);
+			}
+		}
+	}
+}
+
+
+void FInternalPlayWorldCommandCallbacks::RepeatLastLaunch_Clicked()
+{
+	const ULevelEditorPlaySettings* PlaySettings = GetDefault<ULevelEditorPlaySettings>();
+
+	switch (PlaySettings->LastExecutedLaunchModeType)
+	{
+	case LaunchMode_OnDevice:
+		if (FParse::Param(FCommandLine::Get(), TEXT("OldLaunchOn")))
+		{
+			if (IsReadyToLaunchOnDevice(PlaySettings->LastExecutedLaunchDevice))
+			{
+				LaunchOnDevice(PlaySettings->LastExecutedLaunchDevice, PlaySettings->LastExecutedLaunchName, false);
+			}
+		}
+		else
+		{
+			// @todo turnkey - remove this function completely and all the platform support, and put it into Turnkey
+//			if (IsReadyToLaunchOnDevice(PlaySettings->LastExecutedLaunchDevice))
+			{
+				LaunchOnDevice(PlaySettings->LastExecutedLaunchDevice, PlaySettings->LastExecutedLaunchName, true);
+			}
+		}
+
+
+		break;
+
+	default:
+		break;
+	}
+}
+
+
+bool FInternalPlayWorldCommandCallbacks::RepeatLastLaunch_CanExecute()
+{
+	const ULevelEditorPlaySettings* PlaySettings = GetDefault<ULevelEditorPlaySettings>();
+
+	switch (PlaySettings->LastExecutedLaunchModeType)
+	{
+	case LaunchMode_OnDevice:
+		return CanLaunchOnDevice(PlaySettings->LastExecutedLaunchName);
+
+	default:
+		return false;
+	}
+}
+
+
+FText FInternalPlayWorldCommandCallbacks::GetRepeatLastLaunchToolTip()
+{
+	const ULevelEditorPlaySettings* PlaySettings = GetDefault<ULevelEditorPlaySettings>();
+
+	switch (PlaySettings->LastExecutedLaunchModeType)
+	{
+	case LaunchMode_OnDevice:
+		if (CanLaunchOnDevice(PlaySettings->LastExecutedLaunchName))
+		{
+			FFormatNamedArguments Arguments;
+			Arguments.Add(TEXT("DeviceName"), FText::FromString(PlaySettings->LastExecutedLaunchName));
+
+			return FText::Format(LOCTEXT("RepeatLaunchTooltip", "Launch this level on {DeviceName}"), Arguments);
+		}
+
+		break;
+
+	default:
+		break;
+	}
+
+	return LOCTEXT("RepeatLaunchSelectOptionToolTip", "Select a play-on target from the combo menu");
+}
+
+
+FSlateIcon FInternalPlayWorldCommandCallbacks::GetRepeatLastLaunchIcon()
+{
+	const ULevelEditorPlaySettings* PlaySettings = GetDefault<ULevelEditorPlaySettings>();
+
+	// @todo gmp: add play mode specific icons
+	switch (PlaySettings->LastExecutedLaunchModeType)
+	{
+	case LaunchMode_OnDevice:
+		break;
+	}
+
+	static FName RepeatLastLaunchIcon("PlayWorld.RepeatLastLaunch");
+
+	return FSlateIcon(FEditorStyle::GetStyleSetName(), RepeatLastLaunchIcon);
+}
+
+bool FInternalPlayWorldCommandCallbacks::IsReadyToLaunchOnDevice(FString DeviceId)
+{
+	int32 Index = 0;
+	DeviceId.FindChar(TEXT('@'), Index);
+	FName PlatformName = FName(*DeviceId.Left(Index));
+
+	const PlatformInfo::FTargetPlatformInfo* PlatformInfo = PlatformInfo::FindPlatformInfo(PlatformName);
+	checkf(PlatformInfo, TEXT("Unable to find PlatformInfo for %s"), *PlatformName.ToString());
+
+	FGameProjectGenerationModule& GameProjectModule = FModuleManager::LoadModuleChecked<FGameProjectGenerationModule>(TEXT("GameProjectGeneration"));
+	bool bHasCode = GameProjectModule.Get().ProjectHasCodeFiles();
+
+	if (ITurnkeySupportModule::Get().GetSdkInfo(PlatformName).Status != ETurnkeyPlatformSdkStatus::Valid)
+	{
+		IMainFrameModule& MainFrameModule = FModuleManager::GetModuleChecked<IMainFrameModule>(TEXT("MainFrame"));
+		MainFrameModule.BroadcastMainFrameSDKNotInstalled(PlatformInfo->Name.ToString(), PlatformInfo->DataDrivenPlatformInfo->SDKTutorial);
+		TArray<FAnalyticsEventAttribute> ParamArray;
+		ParamArray.Add(FAnalyticsEventAttribute(TEXT("Time"), 0.0));
+		FEditorAnalytics::ReportEvent(TEXT("Editor.LaunchOn.Failed"), PlatformInfo->Name.ToString(), bHasCode, EAnalyticsErrorCodes::SDKNotFound, ParamArray);
+		return false;
+	}
+
+	const ITargetPlatform* Platform = GetTargetPlatformManager()->FindTargetPlatform(PlatformName);
+	if (Platform)
+	{
+		FString NotInstalledTutorialLink;
+		FString DocumentationLink;
+		FText CustomizedLogMessage;
+
+		EBuildConfiguration BuildConfiguration = GetDefault<ULevelEditorPlaySettings>()->GetLaunchBuildConfiguration();
+		bool bEnableAssetNativization = false;
+		int32 Result = Platform->CheckRequirements(bHasCode, BuildConfiguration, bEnableAssetNativization, NotInstalledTutorialLink, DocumentationLink, CustomizedLogMessage);
+
+		// report to analytics
+		FEditorAnalytics::ReportBuildRequirementsFailure(TEXT("Editor.LaunchOn.Failed"), PlatformName.ToString(), bHasCode, Result);
+
+		// report to message log
+		bool UnrecoverableError = false;
+
+		if ((Result & ETargetPlatformReadyStatus::SDKNotFound) != 0)
+		{
+			AddMessageLog(
+				LOCTEXT("SdkNotFoundMessage", "Software Development Kit (SDK) not found."),
+				CustomizedLogMessage.IsEmpty() ? FText::Format(LOCTEXT("SdkNotFoundMessageDetail", "Please install the SDK for the {0} target platform!"), Platform->DisplayName()) : CustomizedLogMessage,
+				NotInstalledTutorialLink,
+				DocumentationLink
+			);
+
+			UnrecoverableError = true;
+		}
+
+		if ((Result & ETargetPlatformReadyStatus::LicenseNotAccepted) != 0)
+		{
+			AddMessageLog(
+				LOCTEXT("LicenseNotAcceptedMessage", "License not accepted."),
+				CustomizedLogMessage.IsEmpty() ? LOCTEXT("LicenseNotAcceptedMessageDetail", "License must be accepted in project settings to deploy your app to the device.") : CustomizedLogMessage,
+				NotInstalledTutorialLink,
+				DocumentationLink
+			);
+
+			UnrecoverableError = true;
+		}
+
+		if ((Result & ETargetPlatformReadyStatus::ProvisionNotFound) != 0)
+		{
+			AddMessageLog(
+				LOCTEXT("ProvisionNotFoundMessage", "Provision not found."),
+				CustomizedLogMessage.IsEmpty() ? LOCTEXT("ProvisionNotFoundMessageDetail", "A provision is required for deploying your app to the device.") : CustomizedLogMessage,
+				NotInstalledTutorialLink,
+				DocumentationLink
+			);
+
+			UnrecoverableError = true;
+		}
+
+		if ((Result & ETargetPlatformReadyStatus::SigningKeyNotFound) != 0)
+		{
+			AddMessageLog(
+				LOCTEXT("SigningKeyNotFoundMessage", "Signing key not found."),
+				CustomizedLogMessage.IsEmpty() ? LOCTEXT("SigningKeyNotFoundMessageDetail", "The app could not be digitally signed, because the signing key is not configured.") : CustomizedLogMessage,
+				NotInstalledTutorialLink,
+				DocumentationLink
+			);
+
+			UnrecoverableError = true;
+		}
+
+		if ((Result & ETargetPlatformReadyStatus::ManifestNotFound) != 0)
+		{
+			AddMessageLog(
+				LOCTEXT("ManifestNotFound", "Manifest not found."),
+				CustomizedLogMessage.IsEmpty() ? LOCTEXT("ManifestNotFoundMessageDetail", "The generated application manifest could not be found.") : CustomizedLogMessage,
+				NotInstalledTutorialLink,
+				DocumentationLink
+			);
+
+			UnrecoverableError = true;
+		}
+
+		if ((Result & ETargetPlatformReadyStatus::RemoveServerNameEmpty) != 0
+			&& (bHasCode || (Result & ETargetPlatformReadyStatus::CodeBuildRequired)
+				|| (!FApp::GetEngineIsPromotedBuild() && !FApp::IsEngineInstalled())))
+		{
+			AddMessageLog(
+				LOCTEXT("RemoveServerNameNotFound", "Remote compiling requires a server name. "),
+				CustomizedLogMessage.IsEmpty() ? LOCTEXT("RemoveServerNameNotFoundDetail", "Please specify one in the Remote Server Name settings field.") : CustomizedLogMessage,
+				NotInstalledTutorialLink,
+				DocumentationLink
+			);
+			UnrecoverableError = true;
+		}
+
+		if (UnrecoverableError)
+		{
+			return false;
+		}
+
+		// report to main frame
+		if ((Result & ETargetPlatformReadyStatus::CodeUnsupported) != 0)
+		{
+			// show the message
+			FMessageDialog::Open(EAppMsgType::Ok, LOCTEXT("NotSupported_CodeBased", "Sorry, launching a code-based project for the selected platform is currently not supported. This feature may be available in a future release."));
+			return false;
+		}
+		if ((Result & ETargetPlatformReadyStatus::PluginsUnsupported) != 0)
+		{
+			// show the message
+			FMessageDialog::Open(EAppMsgType::Ok, LOCTEXT("NotSupported_Plugins", "Sorry, launching a project with third-party plugins is currently not supported for the selected platform. This feature may be available in a future release."));
+			return false;
+		}
+	}
+	else
+	{
+		IMainFrameModule& MainFrameModule = FModuleManager::GetModuleChecked<IMainFrameModule>(TEXT("MainFrame"));
+		MainFrameModule.BroadcastMainFrameSDKNotInstalled(PlatformInfo->Name.ToString(), PlatformInfo->DataDrivenPlatformInfo->SDKTutorial);
+		return false;
+	}
+
+	return true;
+}
+
+
+
+#endif
 
 
 
diff --git a/Engine/Source/Editor/TurnkeySupport/Private/TurnkeySupportModule.h b/Engine/Source/Editor/TurnkeySupport/Private/TurnkeySupportModule.h
--- a/Engine/Source/Editor/TurnkeySupport/Private/TurnkeySupportModule.h
+++ b/Engine/Source/Editor/TurnkeySupport/Private/TurnkeySupportModule.h
@@ -13,20 +13,13 @@ class FTurnkeySupportModule	: public ITurnkeySupportModule
 {
 public:
 
-	/**
-	 *
-	 * @return	The newly-created menu widget
-	 */
-	virtual TSharedRef<class SWidget> MakeTurnkeyMenu() const override;
-
-	/**
-	 * @return	The newly-created menu widget
-	 */
-	virtual void UpdateSdkInfo() override;
+	virtual void MakeTurnkeyMenu(struct FToolMenuSection& MenuSection) const override;
+	virtual void MakeQuickLaunchItems(class UToolMenu* Menu, FOnQuickLaunchSelected ExternalOnClickDelegate) const override;
+	virtual void RepeatQuickLaunch(FString DeviceId) override;
+
+
 
-	/**
-	 * Runs Turnkey to get the Sdk information for a list of devices
-	 */
+	virtual void UpdateSdkInfo() override;
 	virtual void UpdateSdkInfoForDevices(TArray<FString> DeviceIds) override;
 
 	virtual FTurnkeySdkInfo GetSdkInfo(FName PlatformName, bool bBlockIfQuerying = true) const override;
@@ -52,4 +45,7 @@ private:
 	// Information about the validity of each connected device (by string, discovered by Turnkey)
 	TMap<FString, FTurnkeySdkInfo> PerDeviceSdkInfo;
 
+
+	// menu helpers
+	TSharedRef<SWidget> MakeTurnkeyMenuWidget() const;
 };
diff --git a/Engine/Source/Editor/TurnkeySupport/Public/Interfaces/ITurnkeySupportModule.h b/Engine/Source/Editor/TurnkeySupport/Public/Interfaces/ITurnkeySupportModule.h
--- a/Engine/Source/Editor/TurnkeySupport/Public/Interfaces/ITurnkeySupportModule.h
+++ b/Engine/Source/Editor/TurnkeySupport/Public/Interfaces/ITurnkeySupportModule.h
@@ -5,6 +5,7 @@
 #include "CoreMinimal.h"
 #include "Modules/ModuleInterface.h"
 #include "Modules/ModuleManager.h"
+#include "Modules/ModuleManager.h"
 
 class SWidget;
 
@@ -29,9 +30,15 @@ struct FTurnkeySdkInfo
 	FString AutoSDKVersion; // only valid for platform, not device
 	FString MinAllowedVersion;
 	FString MaxAllowedVersion;
+	bool bCanInstallFullSdk;
+	bool bCanInstallAutoSdk;
 };
 
 
+
+DECLARE_DELEGATE_OneParam(FOnQuickLaunchSelected, FString);
+
+
 /**
  * Interface for turnkey support module
  */
@@ -40,13 +47,23 @@ class ITurnkeySupportModule
 {
 public:
 
+	/**
+	 * Make a Platforms menu in the given MenuSection
+	 */
+	virtual void MakeTurnkeyMenu(struct FToolMenuSection& MenuSection) const = 0;
+
+	/**
+	 * Make menu items for Quick Launch, so they can be added to the Play menu
+	 */
+	virtual void MakeQuickLaunchItems(class UToolMenu* Menu, FOnQuickLaunchSelected ExternalOnClickDelegate) const = 0;
+
 	/**
 	 * Runs Turnkey to get the Sdk information for all known platforms
 	 */
-	virtual TSharedRef<SWidget> MakeTurnkeyMenu() const = 0;
+	virtual void RepeatQuickLaunch(FString DeviceId) = 0;
 
 	/**
-	 * @return	The newly-created menu widget
+	 * Runs Turnkey to get the Sdk information for all known platforms
 	 */
 	virtual void UpdateSdkInfo() = 0;
 
diff --git a/Engine/Source/Editor/UnrealEd/Classes/Settings/LevelEditorPlaySettings.h b/Engine/Source/Editor/UnrealEd/Classes/Settings/LevelEditorPlaySettings.h
--- a/Engine/Source/Editor/UnrealEd/Classes/Settings/LevelEditorPlaySettings.h
+++ b/Engine/Source/Editor/UnrealEd/Classes/Settings/LevelEditorPlaySettings.h
@@ -80,6 +80,9 @@ enum EPlayModeType
 	/** Simulates in viewport without possessing the player. */
 	PlayMode_Simulate,
 
+	/** Runs the last launched device (from Platforms menu) */
+	PlayMode_QuickLaunch,
+
 	/** The number of different Play Modes. */
 	PlayMode_Count,
 };
diff --git a/Engine/Source/Editor/UnrealEd/Private/Kismet2/DebuggerCommands.cpp b/Engine/Source/Editor/UnrealEd/Private/Kismet2/DebuggerCommands.cpp
--- a/Engine/Source/Editor/UnrealEd/Private/Kismet2/DebuggerCommands.cpp
+++ b/Engine/Source/Editor/UnrealEd/Private/Kismet2/DebuggerCommands.cpp
@@ -156,11 +156,6 @@ public:
 
 	static void HandleShowSDKTutorial(FString PlatformName, FString NotInstalledDocLink);
 
-	static void RepeatLastLaunch_Clicked();
-	static bool RepeatLastLaunch_CanExecute();
-	static FText GetRepeatLastLaunchToolTip();
-	static FSlateIcon GetRepeatLastLaunchIcon();
-
 	static FSlateIcon GetResumePlaySessionImage();
 	static FText GetResumePlaySessionToolTip();
 	static void StopPlaySession_Clicked();
@@ -353,9 +348,6 @@ void FPlayWorldCommands::RegisterCommands()
 	UI_COMMAND(StepOver, "Step Over", "Step to the next node to be executed in the current graph", EUserInterfaceActionType::Button, FInputChord(EKeys::F10));
 	UI_COMMAND(StepOut, "Step Out", "Step Out to the next node to be executed in the parent graph", EUserInterfaceActionType::Button, FInputChord(EModifierKey::Alt | EModifierKey::Shift, EKeys::F11));
 
-	// Launch
-	UI_COMMAND(RepeatLastLaunch, "Launch", "Launches the game on the device as the last session launched from the dropdown next to the Play on Device button on the level editor toolbar", EUserInterfaceActionType::Button, FInputChord(EKeys::P, EModifierKey::Alt | EModifierKey::Shift))
-
 	// PIE mobile preview devices.
 	AddPIEPreviewDeviceCommands();
 }
@@ -473,15 +465,6 @@ void FPlayWorldCommands::BindGlobalPlayWorldCommands()
 		FExecuteAction::CreateStatic(&FInternalPlayWorldCommandCallbacks::PlayInSettings_Clicked)
 	);
 
-	// Launch
-
-	ActionList.MapAction(Commands.RepeatLastLaunch,
-		FExecuteAction::CreateStatic(&FInternalPlayWorldCommandCallbacks::RepeatLastLaunch_Clicked),
-		FCanExecuteAction::CreateStatic(&FInternalPlayWorldCommandCallbacks::RepeatLastLaunch_CanExecute),
-		FIsActionChecked(),
-		FIsActionButtonVisible::CreateStatic(&FInternalPlayWorldCommandCallbacks::CanShowNonPlayWorldOnlyActions)
-	);
-
 
 	// Stop play session
 	ActionList.MapAction(Commands.StopPlaySession,
@@ -630,34 +613,7 @@ void FPlayWorldCommands::BuildToolbar(FToolMenuSection& InSection, bool bInclude
 
 	if (bIncludeLaunchButtonAndOptions)
 	{
-		InSection.AddDynamicEntry("LaunchButtons", FNewToolMenuSectionDelegate::CreateLambda([](FToolMenuSection& InDynamicSection)
-		{
-			if (GetDefault<UEditorStyleSettings>()->bShowLaunchMenus)
-			{
-				// Launch
-				InDynamicSection.AddEntry(FToolMenuEntry::InitToolBarButton(
-					FPlayWorldCommands::Get().RepeatLastLaunch,
-					LOCTEXT("RepeatLastLaunch", "Launch"),
-					TAttribute< FText >::Create(TAttribute< FText >::FGetter::CreateStatic(&FInternalPlayWorldCommandCallbacks::GetRepeatLastLaunchToolTip)),
-					TAttribute< FSlateIcon >::Create(TAttribute< FSlateIcon >::FGetter::CreateStatic(&FInternalPlayWorldCommandCallbacks::GetRepeatLastLaunchIcon)),
-					FName(TEXT("RepeatLastLaunch"))
-				));
-
-				// Launch combo box
-				FUIAction LaunchMenuAction;
-				LaunchMenuAction.IsActionVisibleDelegate = FIsActionButtonVisible::CreateStatic(&FInternalPlayWorldCommandCallbacks::CanShowNonPlayWorldOnlyActions);
-
-				InDynamicSection.AddEntry(FToolMenuEntry::InitComboButton(
-					"TurnkeyCombo",
-					LaunchMenuAction,
-					FOnGetContent::CreateLambda([] { return ITurnkeySupportModule::Get().MakeTurnkeyMenu(); }),
-					LOCTEXT("LaunchCombo_Label", "Launch Options"),
-					LOCTEXT("PlayWithTurnkeyComboToolTip", "Options for launching on a device and manage platforms"),
-					FSlateIcon(FEditorStyle::GetStyleSetName(), "LevelEditor.RepeatLastLaunch"),
-					true
-				));
-			}
-		}));
+		ITurnkeySupportModule::Get().MakeTurnkeyMenu(InSection);
 	}
 
 	// Resume/pause toggle (only one will be visible, and only in PIE/SIE)
@@ -816,6 +772,33 @@ static void MakePreviewDeviceMenu(FMenuBuilder& MenuBuilder)
 	}
 }
 
+void SetLastExecutedPlayMode(EPlayModeType PlayMode)
+{
+	ULevelEditorPlaySettings* PlaySettings = GetMutableDefault<ULevelEditorPlaySettings>();
+	PlaySettings->LastExecutedPlayModeType = PlayMode;
+
+	FPropertyChangedEvent PropChangeEvent(ULevelEditorPlaySettings::StaticClass()->FindPropertyByName(GET_MEMBER_NAME_CHECKED(ULevelEditorPlaySettings, LastExecutedPlayModeType)));
+	PlaySettings->PostEditChangeProperty(PropChangeEvent);
+
+	PlaySettings->SaveConfig();
+}
+
+
+static void RememberQuickLaunch(FString DeviceId)
+{
+	// remember that clicking Play should launch
+	SetLastExecutedPlayMode(EPlayModeType::PlayMode_QuickLaunch);
+
+	// store the device name in the play settings for next click/run
+	ULevelEditorPlaySettings* PlaySettings = GetMutableDefault<ULevelEditorPlaySettings>();
+
+	PlaySettings->LastExecutedLaunchName = DeviceId;
+
+	FPropertyChangedEvent PropChangeEvent(ULevelEditorPlaySettings::StaticClass()->FindPropertyByName(GET_MEMBER_NAME_CHECKED(ULevelEditorPlaySettings, LastExecutedLaunchName)));
+	PlaySettings->PostEditChangeProperty(PropChangeEvent);
+	PlaySettings->SaveConfig();
+}
+
 TSharedRef< SWidget > FPlayWorldCommands::GeneratePlayMenuContent(TSharedRef<FUICommandList> InCommandList)
 {
 	static const FName MenuName("UnrealEd.PlayWorldCommands.PlayMenu");
@@ -892,14 +875,18 @@ TSharedRef< SWidget > FPlayWorldCommands::GeneratePlayMenuContent(TSharedRef<FUI
 			FLocal::AddPlayModeMenuEntry(Section, PlayMode_Simulate);
 		}
 
+		// quick launch on devices
+ 		ITurnkeySupportModule::Get().MakeQuickLaunchItems(Menu, FOnQuickLaunchSelected::CreateStatic(&RememberQuickLaunch));
+
 		// tip section
 		{
 			FToolMenuSection& Section = Menu->AddSection("LevelEditorPlayTip");
+			Section.AddSeparator(NAME_None);
 			Section.AddEntry(FToolMenuEntry::InitWidget(
 				"PlayIn",
 				SNew(STextBlock)
 				.ColorAndOpacity(FSlateColor::UseSubduedForeground())
-				.Text(LOCTEXT("PlayInTip", "Launching a game preview with a different mode will change your default 'Play' mode in the toolbar"))
+				.Text(LOCTEXT("PlayInTip", "Launching a game (preview or on device) with a different mode will change your default 'Play' mode in the toolbar"))
 				.WrapTextAt(250),
 				FText::GetEmpty()));
 		}
@@ -1151,17 +1138,6 @@ bool FInternalPlayWorldCommandCallbacks::CanShowLateJoinButton()
 	return GetDefault<UEditorExperimentalSettings>()->bAllowLateJoinInPIE && HasPlayWorld();
 }
 
-void SetLastExecutedPlayMode(EPlayModeType PlayMode)
-{
-	ULevelEditorPlaySettings* PlaySettings = GetMutableDefault<ULevelEditorPlaySettings>();
-	PlaySettings->LastExecutedPlayModeType = PlayMode;
-
-	FPropertyChangedEvent PropChangeEvent(ULevelEditorPlaySettings::StaticClass()->FindPropertyByName(GET_MEMBER_NAME_CHECKED(ULevelEditorPlaySettings, LastExecutedPlayModeType)));
-	PlaySettings->PostEditChangeProperty(PropChangeEvent);
-
-	PlaySettings->SaveConfig();
-}
-
 
 void FInternalPlayWorldCommandCallbacks::Simulate_Clicked()
 {
@@ -1352,32 +1328,65 @@ void SetLastExecutedLaunchMode(ELaunchModeType LaunchMode)
 
 void FInternalPlayWorldCommandCallbacks::RepeatLastPlay_Clicked()
 {
-	// Let a game have a go at settings before we play
 	ULevelEditorPlaySettings* PlaySettings = GetMutableDefault<ULevelEditorPlaySettings>();
 	PlaySettings->PostEditChange();
 
-	// Grab the play command and execute it
-	TSharedRef<FUICommandInfo> LastCommand = GetLastPlaySessionCommand();
-	UE_LOG(LogDebuggerCommands, Log, TEXT("Repeating last play command: %s"), *LastCommand->GetLabel().ToString());
+	// hand over to Turnkey module
+	if (PlaySettings->LastExecutedPlayModeType == EPlayModeType::PlayMode_QuickLaunch)
+	{
+		ITurnkeySupportModule::Get().RepeatQuickLaunch(PlaySettings->LastExecutedLaunchName);
+	}
+	else
+	{
+		// Grab the play command and execute it
+		TSharedRef<FUICommandInfo> LastCommand = GetLastPlaySessionCommand();
+		UE_LOG(LogDebuggerCommands, Log, TEXT("Repeating last play command: %s"), *LastCommand->GetLabel().ToString());
 
-	FPlayWorldCommands::GlobalPlayWorldActions->ExecuteAction(LastCommand);
+		FPlayWorldCommands::GlobalPlayWorldActions->ExecuteAction(LastCommand);
+	}
 }
 
 
 bool FInternalPlayWorldCommandCallbacks::RepeatLastPlay_CanExecute()
 {
+	const ULevelEditorPlaySettings* PlaySettings = GetDefault<ULevelEditorPlaySettings>();
+	if (PlaySettings->LastExecutedPlayModeType == EPlayModeType::PlayMode_QuickLaunch)
+	{
+		// return true, and let Turnkey module determine if it's still usable, and show an error if not
+		return true;
+	}
+
 	return FPlayWorldCommands::GlobalPlayWorldActions->CanExecuteAction(GetLastPlaySessionCommand());
 }
 
 
 FText FInternalPlayWorldCommandCallbacks::GetRepeatLastPlayToolTip()
 {
+	const ULevelEditorPlaySettings* PlaySettings = GetDefault<ULevelEditorPlaySettings>();
+	if (PlaySettings->LastExecutedPlayModeType == EPlayModeType::PlayMode_QuickLaunch)
+	{
+		// @todo make a proper tooltip!
+		return FText::FromString(PlaySettings->LastExecutedLaunchName);
+	}
+
 	return GetLastPlaySessionCommand()->GetDescription();
 }
 
 
 FSlateIcon FInternalPlayWorldCommandCallbacks::GetRepeatLastPlayIcon()
 {
+	// get platform icon for Quick Launch mode
+	const ULevelEditorPlaySettings* PlaySettings = GetDefault<ULevelEditorPlaySettings>();
+	if (PlaySettings->LastExecutedPlayModeType == EPlayModeType::PlayMode_QuickLaunch)
+	{
+		FTargetDeviceId DeviceId;
+		FTargetDeviceId::Parse(PlaySettings->LastExecutedLaunchName, DeviceId);
+
+		// get platform name from DeviceId
+		
+		return FSlateIcon(FEditorStyle::GetStyleSetName(), FDataDrivenPlatformInfoRegistry::GetPlatformInfo(DeviceId.GetPlatformName()).GetIconStyleName(EPlatformIconSize::Large));
+	}
+
 	return GetLastPlaySessionCommand()->GetIcon();
 }
 
@@ -1930,367 +1939,6 @@ EPlayModeLocations FInternalPlayWorldCommandCallbacks::GetPlayModeLocation()
 
 
 
-/////// Will move these into TurnkeySupportModule soon:
-
-
-
-bool FInternalPlayWorldCommandCallbacks::CanLaunchOnDevice(const FString& DeviceName)
-{
-	if (!GUnrealEd->IsPlayingViaLauncher())
-	{
-		static TWeakPtr<ITargetDeviceProxyManager> DeviceProxyManagerPtr;
-
-		if (!DeviceProxyManagerPtr.IsValid())
-		{
-			ITargetDeviceServicesModule* TargetDeviceServicesModule = FModuleManager::Get().LoadModulePtr<ITargetDeviceServicesModule>(TEXT("TargetDeviceServices"));
-			if (TargetDeviceServicesModule)
-			{
-				DeviceProxyManagerPtr = TargetDeviceServicesModule->GetDeviceProxyManager();
-			}
-		}
-
-		TSharedPtr<ITargetDeviceProxyManager> DeviceProxyManager = DeviceProxyManagerPtr.Pin();
-		if (DeviceProxyManager.IsValid())
-		{
-			TSharedPtr<ITargetDeviceProxy> DeviceProxy = DeviceProxyManager->FindProxy(DeviceName);
-			if (DeviceProxy.IsValid() && DeviceProxy->IsConnected() && DeviceProxy->IsAuthorized())
-			{
-				return true;
-			}
-
-			// check if this is an aggregate proxy
-			TArray<TSharedPtr<ITargetDeviceProxy>> Devices;
-			DeviceProxyManager->GetProxies(FName(*DeviceName), false, Devices);
-
-			// returns true if the game can be launched al least on 1 device
-			for (auto DevicesIt = Devices.CreateIterator(); DevicesIt; ++DevicesIt)
-			{
-				TSharedPtr<ITargetDeviceProxy> DeviceAggregateProxy = *DevicesIt;
-				if (DeviceAggregateProxy.IsValid() && DeviceAggregateProxy->IsConnected() && DeviceAggregateProxy->IsAuthorized())
-				{
-					return true;
-				}
-			}
-
-		}
-	}
-
-	return false;
-}
-
-
-void FInternalPlayWorldCommandCallbacks::LaunchOnDevice(const FString& DeviceId, const FString& DeviceName, bool bUseTurnkey)
-{
-	FTargetDeviceId TargetDeviceId;
-	if (FTargetDeviceId::Parse(DeviceId, TargetDeviceId))
-	{
-		const PlatformInfo::FTargetPlatformInfo* const PlatformInfo = PlatformInfo::FindPlatformInfo(*TargetDeviceId.GetPlatformName());
-		FString UBTPlatformName = PlatformInfo->DataDrivenPlatformInfo->UBTPlatformString;
-		FString IniPlatformName = PlatformInfo->IniPlatformName.ToString();
-
-		check(PlatformInfo);
-
-		if (FInstalledPlatformInfo::Get().IsPlatformMissingRequiredFile(UBTPlatformName))
-		{
-			if (!FInstalledPlatformInfo::OpenInstallerOptions())
-			{
-				FMessageDialog::Open(EAppMsgType::Ok, LOCTEXT("MissingPlatformFilesLaunch", "Missing required files to launch on this platform."));
-			}
-			return;
-		}
-
-		if (FModuleManager::LoadModuleChecked<IProjectTargetPlatformEditorModule>("ProjectTargetPlatformEditor").ShowUnsupportedTargetWarning(*TargetDeviceId.GetPlatformName()))
-		{
-			GUnrealEd->CancelPlayingViaLauncher();
-
-			FRequestPlaySessionParams::FLauncherDeviceInfo DeviceInfo;
-			DeviceInfo.DeviceId = DeviceId;
-			DeviceInfo.DeviceName = DeviceName;
-			// @todo turnkey: we set this to false because we will kick off a Turnkey run before cooking, etc, to get an early warning. however, if it's too difficult
-			// to get an error back from CreateUatTask, then we should set this to bUseTurnkey and remove the block below, and let the code in FLauncherWorker::CreateAndExecuteTasks handle it
-			DeviceInfo.bUpdateDeviceFlash = false;
-
-			FRequestPlaySessionParams SessionParams;
-			SessionParams.SessionDestination = EPlaySessionDestinationType::Launcher;
-			SessionParams.LauncherTargetDevice = DeviceInfo;
-
-			// if we want to check device flash before we start cooking, kick it off now. we could delay this 
-			if (bUseTurnkey)
-			{
-				FString CommandLine = FString::Printf(TEXT("Turnkey -command=VerifySdk -UpdateIfNeeded -platform=%s -EditorIO -noturnkeyvariables -device=%s -utf8output -WaitForUATMutex"), *UBTPlatformName, *DeviceName);
-				FText TaskName = LOCTEXT("VerifyingSDK", "Verifying SDK and Device");
-
-				IUATHelperModule::Get().CreateUatTask(CommandLine, FText::FromString(IniPlatformName), TaskName, TaskName, FEditorStyle::GetBrush(TEXT("MainFrame.PackageProject")),
-					[SessionParams](FString Result, double)
-				{
-					// unfortunate string comparison for success
-					bool bWasSuccessful = Result == TEXT("Completed");
-					AsyncTask(ENamedThreads::GameThread, [SessionParams, bWasSuccessful]()
-					{
-						if (bWasSuccessful)
-						{
-							GUnrealEd->RequestPlaySession(SessionParams);
-						}
-						else
-						{
-							TSharedRef<SWindow> Win = OpenMsgDlgInt_NonModal(EAppMsgType::YesNo, LOCTEXT("SDKCheckFailed", "SDK Verification failed. Would you like to attempt the Launch On anyway?"), LOCTEXT("SDKCheckFailedTitle", "SDK Verification"),
-								FOnMsgDlgResult::CreateLambda([SessionParams](const TSharedRef<SWindow>&, EAppReturnType::Type Choice)
-							{
-								if (Choice == EAppReturnType::Yes)
-								{
-									GUnrealEd->RequestPlaySession(SessionParams);
-								}
-							}));
-							Win->ShowWindow();
-						}
-					});
-				});
-			}
-			else
-			{
-				GUnrealEd->RequestPlaySession(SessionParams);
-			}
-		}
-	}
-}
-
-
-void FInternalPlayWorldCommandCallbacks::RepeatLastLaunch_Clicked()
-{
-	const ULevelEditorPlaySettings* PlaySettings = GetDefault<ULevelEditorPlaySettings>();
-
-	switch (PlaySettings->LastExecutedLaunchModeType)
-	{
-	case LaunchMode_OnDevice:
-		if (FParse::Param(FCommandLine::Get(), TEXT("OldLaunchOn")))
-		{
-			if (IsReadyToLaunchOnDevice(PlaySettings->LastExecutedLaunchDevice))
-			{
-				LaunchOnDevice(PlaySettings->LastExecutedLaunchDevice, PlaySettings->LastExecutedLaunchName, false);
-			}
-		}
-		else
-		{
-			// @todo turnkey - remove this function completely and all the platform support, and put it into Turnkey
-//			if (IsReadyToLaunchOnDevice(PlaySettings->LastExecutedLaunchDevice))
-			{
-				LaunchOnDevice(PlaySettings->LastExecutedLaunchDevice, PlaySettings->LastExecutedLaunchName, true);
-			}
-		}
-
-
-		break;
-
-	default:
-		break;
-	}
-}
-
-
-bool FInternalPlayWorldCommandCallbacks::RepeatLastLaunch_CanExecute()
-{
-	const ULevelEditorPlaySettings* PlaySettings = GetDefault<ULevelEditorPlaySettings>();
-
-	switch (PlaySettings->LastExecutedLaunchModeType)
-	{
-	case LaunchMode_OnDevice:
-		return CanLaunchOnDevice(PlaySettings->LastExecutedLaunchName);
-
-	default:
-		return false;
-	}
-}
-
-
-FText FInternalPlayWorldCommandCallbacks::GetRepeatLastLaunchToolTip()
-{
-	const ULevelEditorPlaySettings* PlaySettings = GetDefault<ULevelEditorPlaySettings>();
-
-	switch (PlaySettings->LastExecutedLaunchModeType)
-	{
-	case LaunchMode_OnDevice:
-		if (CanLaunchOnDevice(PlaySettings->LastExecutedLaunchName))
-		{
-			FFormatNamedArguments Arguments;
-			Arguments.Add(TEXT("DeviceName"), FText::FromString(PlaySettings->LastExecutedLaunchName));
-
-			return FText::Format(LOCTEXT("RepeatLaunchTooltip", "Launch this level on {DeviceName}"), Arguments);
-		}
-
-		break;
-
-	default:
-		break;
-	}
-
-	return LOCTEXT("RepeatLaunchSelectOptionToolTip", "Select a play-on target from the combo menu");
-}
-
-
-FSlateIcon FInternalPlayWorldCommandCallbacks::GetRepeatLastLaunchIcon()
-{
-	const ULevelEditorPlaySettings* PlaySettings = GetDefault<ULevelEditorPlaySettings>();
-
-	// @todo gmp: add play mode specific icons
-	switch (PlaySettings->LastExecutedLaunchModeType)
-	{
-	case LaunchMode_OnDevice:
-		break;
-	}
-
-	static FName RepeatLastLaunchIcon("PlayWorld.RepeatLastLaunch");
-
-	return FSlateIcon(FEditorStyle::GetStyleSetName(), RepeatLastLaunchIcon);
-}
-
-bool FInternalPlayWorldCommandCallbacks::IsReadyToLaunchOnDevice(FString DeviceId)
-{
-	int32 Index = 0;
-	DeviceId.FindChar(TEXT('@'), Index);
-	FName PlatformName = FName(*DeviceId.Left(Index));
-
-	const PlatformInfo::FTargetPlatformInfo* PlatformInfo = PlatformInfo::FindPlatformInfo(PlatformName);
-	checkf(PlatformInfo, TEXT("Unable to find PlatformInfo for %s"), *PlatformName.ToString());
-
-	FGameProjectGenerationModule& GameProjectModule = FModuleManager::LoadModuleChecked<FGameProjectGenerationModule>(TEXT("GameProjectGeneration"));
-	bool bHasCode = GameProjectModule.Get().ProjectHasCodeFiles();
-
-	if (ITurnkeySupportModule::Get().GetSdkInfo(PlatformName).Status!= ETurnkeyPlatformSdkStatus::Valid)
-	{
-		IMainFrameModule& MainFrameModule = FModuleManager::GetModuleChecked<IMainFrameModule>(TEXT("MainFrame"));
-		MainFrameModule.BroadcastMainFrameSDKNotInstalled(PlatformInfo->Name.ToString(), PlatformInfo->DataDrivenPlatformInfo->SDKTutorial);
-		TArray<FAnalyticsEventAttribute> ParamArray;
-		ParamArray.Add(FAnalyticsEventAttribute(TEXT("Time"), 0.0));
-		FEditorAnalytics::ReportEvent(TEXT("Editor.LaunchOn.Failed"), PlatformInfo->Name.ToString(), bHasCode, EAnalyticsErrorCodes::SDKNotFound, ParamArray);
-		return false;
-	}
-
-	const ITargetPlatform* Platform = GetTargetPlatformManager()->FindTargetPlatform(PlatformName);
-	if (Platform)
-	{
-		FString NotInstalledTutorialLink;
-		FString DocumentationLink;
-		FText CustomizedLogMessage;
-
-		EBuildConfiguration BuildConfiguration = GetDefault<ULevelEditorPlaySettings>()->GetLaunchBuildConfiguration();
-		bool bEnableAssetNativization = false;
-		int32 Result = Platform->CheckRequirements(bHasCode, BuildConfiguration, bEnableAssetNativization, NotInstalledTutorialLink, DocumentationLink, CustomizedLogMessage);
-
-		// report to analytics
-		FEditorAnalytics::ReportBuildRequirementsFailure(TEXT("Editor.LaunchOn.Failed"), PlatformName.ToString(), bHasCode, Result);
-
-		// report to message log
-		bool UnrecoverableError = false;
-
-		if ((Result & ETargetPlatformReadyStatus::SDKNotFound) != 0)
-		{
-			AddMessageLog(
-				LOCTEXT("SdkNotFoundMessage", "Software Development Kit (SDK) not found."),
-				CustomizedLogMessage.IsEmpty() ? FText::Format(LOCTEXT("SdkNotFoundMessageDetail", "Please install the SDK for the {0} target platform!"), Platform->DisplayName()) : CustomizedLogMessage,
-				NotInstalledTutorialLink,
-				DocumentationLink
-			);
-
-			UnrecoverableError = true;
-		}
-
-		if ((Result & ETargetPlatformReadyStatus::LicenseNotAccepted) != 0)
-		{
-			AddMessageLog(
-				LOCTEXT("LicenseNotAcceptedMessage", "License not accepted."),
-				CustomizedLogMessage.IsEmpty() ? LOCTEXT("LicenseNotAcceptedMessageDetail", "License must be accepted in project settings to deploy your app to the device.") : CustomizedLogMessage,
-				NotInstalledTutorialLink,
-				DocumentationLink
-			);
-
-			UnrecoverableError = true;
-		}
-
-		if ((Result & ETargetPlatformReadyStatus::ProvisionNotFound) != 0)
-		{
-			AddMessageLog(
-				LOCTEXT("ProvisionNotFoundMessage", "Provision not found."),
-				CustomizedLogMessage.IsEmpty() ? LOCTEXT("ProvisionNotFoundMessageDetail", "A provision is required for deploying your app to the device.") : CustomizedLogMessage,
-				NotInstalledTutorialLink,
-				DocumentationLink
-			);
-
-			UnrecoverableError = true;
-		}
-
-		if ((Result & ETargetPlatformReadyStatus::SigningKeyNotFound) != 0)
-		{
-			AddMessageLog(
-				LOCTEXT("SigningKeyNotFoundMessage", "Signing key not found."),
-				CustomizedLogMessage.IsEmpty() ? LOCTEXT("SigningKeyNotFoundMessageDetail", "The app could not be digitally signed, because the signing key is not configured.") : CustomizedLogMessage,
-				NotInstalledTutorialLink,
-				DocumentationLink
-			);
-
-			UnrecoverableError = true;
-		}
-
-		if ((Result & ETargetPlatformReadyStatus::ManifestNotFound) != 0)
-		{
-			AddMessageLog(
-				LOCTEXT("ManifestNotFound", "Manifest not found."),
-				CustomizedLogMessage.IsEmpty() ? LOCTEXT("ManifestNotFoundMessageDetail", "The generated application manifest could not be found.") : CustomizedLogMessage,
-				NotInstalledTutorialLink,
-				DocumentationLink
-			);
-
-			UnrecoverableError = true;
-		}
-
-		if ((Result & ETargetPlatformReadyStatus::RemoveServerNameEmpty) != 0
-			&& (bHasCode || (Result & ETargetPlatformReadyStatus::CodeBuildRequired)
-				|| (!FApp::GetEngineIsPromotedBuild() && !FApp::IsEngineInstalled())))
-		{
-			AddMessageLog(
-				LOCTEXT("RemoveServerNameNotFound", "Remote compiling requires a server name. "),
-				CustomizedLogMessage.IsEmpty() ? LOCTEXT("RemoveServerNameNotFoundDetail", "Please specify one in the Remote Server Name settings field.") : CustomizedLogMessage,
-				NotInstalledTutorialLink,
-				DocumentationLink
-			);
-			UnrecoverableError = true;
-		}
-
-		if (UnrecoverableError)
-		{
-			return false;
-		}
-
-		// report to main frame
-		if ((Result & ETargetPlatformReadyStatus::CodeUnsupported) != 0)
-		{
-			// show the message
-			FMessageDialog::Open(EAppMsgType::Ok, LOCTEXT("NotSupported_CodeBased", "Sorry, launching a code-based project for the selected platform is currently not supported. This feature may be available in a future release."));
-			return false;
-		}
-		if ((Result & ETargetPlatformReadyStatus::PluginsUnsupported) != 0)
-		{
-			// show the message
-			FMessageDialog::Open(EAppMsgType::Ok, LOCTEXT("NotSupported_Plugins", "Sorry, launching a project with third-party plugins is currently not supported for the selected platform. This feature may be available in a future release."));
-			return false;
-		}
-	}
-	else
-	{
-		IMainFrameModule& MainFrameModule = FModuleManager::GetModuleChecked<IMainFrameModule>(TEXT("MainFrame"));
-		MainFrameModule.BroadcastMainFrameSDKNotInstalled(PlatformInfo->Name.ToString(), PlatformInfo->DataDrivenPlatformInfo->SDKTutorial);
-		return false;
-	}
-
-	return true;
-}
-
-
-
-
-
-
-
-
 
 #undef LOCTEXT_NAMESPACE
 
diff --git a/Engine/Source/Editor/UnrealEd/Public/Kismet2/DebuggerCommands.h b/Engine/Source/Editor/UnrealEd/Public/Kismet2/DebuggerCommands.h
--- a/Engine/Source/Editor/UnrealEd/Public/Kismet2/DebuggerCommands.h
+++ b/Engine/Source/Editor/UnrealEd/Public/Kismet2/DebuggerCommands.h
@@ -108,9 +108,6 @@ public:
 	TSharedPtr<FUICommandInfo> StepOut;
 	TSharedPtr<FUICommandInfo> GetMouseControl;
 
-	/** Launch on device */
-	TSharedPtr<FUICommandInfo> RepeatLastLaunch;
-
 protected:
 
 	/** A weak pointer to the current active widget that processes PIE actions */
diff --git a/Engine/Source/Runtime/Core/Private/Misc/DataDrivenPlatformInfoRegistry.cpp b/Engine/Source/Runtime/Core/Private/Misc/DataDrivenPlatformInfoRegistry.cpp
--- a/Engine/Source/Runtime/Core/Private/Misc/DataDrivenPlatformInfoRegistry.cpp
+++ b/Engine/Source/Runtime/Core/Private/Misc/DataDrivenPlatformInfoRegistry.cpp
@@ -276,7 +276,6 @@ static void LoadDDPIIniSettings(const FConfigFile& IniFile, FDataDrivenPlatformI
 	DDPIGetString(IniFile, TEXT("TutorialPath"), Info.SDKTutorial);
 	DDPIGetName(IniFile, TEXT("PlatformGroupName"), Info.PlatformGroupName);
 	DDPIGetName(IniFile, TEXT("PlatformSubMenu"), Info.PlatformSubMenu);
-	DDPIGetString(IniFile, TEXT("PrepareForDebuggingOptions"), Info.PrepareForDebuggingOptions);
 
 
 	DDPIGetString(IniFile, TEXT("NormalIconPath"), Info.IconPaths.NormalPath);
@@ -439,6 +438,10 @@ const FDataDrivenPlatformInfo& FDataDrivenPlatformInfoRegistry::GetPlatformInfo(
 {
 	const FDataDrivenPlatformInfo* Info = GetAllPlatformInfos().Find(PlatformName);
 	static FDataDrivenPlatformInfo Empty;
+	if (Info == nullptr && PlatformName == "Win64")
+	{
+		Info = GetAllPlatformInfos().Find("Windows");
+	}
 	return Info ? *Info : Empty;
 }
 
diff --git a/Engine/Source/Runtime/Core/Public/Misc/DataDrivenPlatformInfoRegistry.h b/Engine/Source/Runtime/Core/Public/Misc/DataDrivenPlatformInfoRegistry.h
--- a/Engine/Source/Runtime/Core/Public/Misc/DataDrivenPlatformInfoRegistry.h
+++ b/Engine/Source/Runtime/Core/Public/Misc/DataDrivenPlatformInfoRegistry.h
@@ -161,9 +161,6 @@ public:
 	FName UBTPlatformName;
 	FString UBTPlatformString;
 
-	/** If this is set, then it is used to prepare a project for debugging */
-	FString PrepareForDebuggingOptions;
-
 	/** Whether or not the platform can use Crash Reporter */
 	bool bCanUseCrashReporter;
 
