- Added -ScriptsForProject to Turnkey calls to reduce the scope of compiled projects

diff --git a/Engine/Source/Editor/TurnkeySupport/Private/TurnkeySupportModule.cpp b/Engine/Source/Editor/TurnkeySupport/Private/TurnkeySupportModule.cpp
--- a/Engine/Source/Editor/TurnkeySupport/Private/TurnkeySupportModule.cpp
+++ b/Engine/Source/Editor/TurnkeySupport/Private/TurnkeySupportModule.cpp
@@ -711,8 +711,6 @@ static void ShowInstallationHelp(FName IniPlatformName)
 
 static void TurnkeyInstallSdk(FString PlatformName, bool bPreferFull, bool bForceInstall, FString DeviceId)
 {
-	//	FString CommandLine = FString::Printf(TEXT("Turnkey -command=InstallSdk -platform=%s -BestAvailable -AllowAutoSdk -EditorIO -noturnkeyvariables -utf8output -WaitForUATMutex"), *PlatformName);
-
 	FString OptionalOptions;
 	if (bPreferFull)
 	{
@@ -727,7 +725,8 @@ static void TurnkeyInstallSdk(FString PlatformName, bool bPreferFull, bool bForc
 		OptionalOptions += FString::Printf(TEXT(" -Device=%s"), *DeviceId);
 	}
 
-	FString CommandLine = FString::Printf(TEXT("Turnkey -command=VerifySdk -UpdateIfNeeded -platform=%s %s -EditorIO -noturnkeyvariables -utf8output -WaitForUATMutex"), *PlatformName, *OptionalOptions);
+	const FString ProjectPath = FPaths::IsProjectFilePathSet() ? FPaths::ConvertRelativePathToFull(FPaths::GetProjectFilePath()) : FPaths::RootDir() / FApp::GetProjectName() / FApp::GetProjectName() + TEXT(".uproject");
+	FString CommandLine = FString::Printf(TEXT("-ScriptsForProject=\"%s\" Turnkey -command=VerifySdk -UpdateIfNeeded -platform=%s %s -EditorIO -noturnkeyvariables -utf8output -WaitForUATMutex"), *ProjectPath, *PlatformName, *OptionalOptions);
 
 	FText TaskName = LOCTEXT("InstallingSdk", "Installing Sdk");
 	FTurnkeyEditorSupport::RunUAT(CommandLine, FText::FromString(PlatformName), TaskName, TaskName, FEditorStyle::GetBrush(TEXT("MainFrame.PackageProject")),
@@ -1482,7 +1481,9 @@ static void PrepForTurnkeyReport(FString& Command, FString& BaseCommandline, FSt
 
 	Command = TEXT("{EngineDir}Build/BatchFiles/RunuAT");
 	//	Command = TEXT("{EngineDir}/Binaries/DotNET/AutomationTool.exe");
-	BaseCommandline = FString::Printf(TEXT("Turnkey -utf8output -WaitForUATMutex -command=VerifySdk -ReportFilename=\"%s\" -log=\"%s\""), *ReportFilename, *LogFilename);
+
+	const FString ProjectPath = FPaths::IsProjectFilePathSet() ? FPaths::ConvertRelativePathToFull(FPaths::GetProjectFilePath()) : FPaths::RootDir() / FApp::GetProjectName() / FApp::GetProjectName() + TEXT(".uproject");
+	BaseCommandline = FString::Printf(TEXT("-ScriptsForProject=\"%s\" Turnkey -utf8output -WaitForUATMutex -command=VerifySdk -ReportFilename=\"%s\" -log=\"%s\""), *ProjectPath, *ReportFilename, *LogFilename);
 
 	// convert into appropriate calls for the current platform
 	FPlatformProcess::ModifyCreateProcParams(Command, BaseCommandline, FGenericPlatformProcess::ECreateProcHelperFlags::AppendPlatformScriptExtension | FGenericPlatformProcess::ECreateProcHelperFlags::RunThroughShell);
