- Hide the Platforms menu during PIE, like the old LaunchOn menu

diff --git a/Engine/Source/Editor/TurnkeySupport/Private/TurnkeyEditorSupport.cpp b/Engine/Source/Editor/TurnkeySupport/Private/TurnkeyEditorSupport.cpp
--- a/Engine/Source/Editor/TurnkeySupport/Private/TurnkeyEditorSupport.cpp
+++ b/Engine/Source/Editor/TurnkeySupport/Private/TurnkeyEditorSupport.cpp
@@ -224,5 +224,13 @@ void FTurnkeyEditorSupport::ShowInstallationHelp(FName IniPlatformName, FString
 #endif
 }
 
+bool FTurnkeyEditorSupport::IsPIERunning()
+{
+#if WITH_EDITOR
+	return GEditor->PlayWorld != NULL;
+#endif
+	return false;
+}
+
 
 #undef LOCTEXT_NAMESPACE
diff --git a/Engine/Source/Editor/TurnkeySupport/Private/TurnkeyEditorSupport.h b/Engine/Source/Editor/TurnkeySupport/Private/TurnkeyEditorSupport.h
--- a/Engine/Source/Editor/TurnkeySupport/Private/TurnkeyEditorSupport.h
+++ b/Engine/Source/Editor/TurnkeySupport/Private/TurnkeyEditorSupport.h
@@ -22,5 +22,5 @@ public:
 	static void ShowRestartToast();
 	static bool CheckSupportedPlatforms(FName IniPlatformName);
 	static void ShowInstallationHelp(FName IniPlatformName, FString DocLink);
-
+	static bool IsPIERunning();
 };
diff --git a/Engine/Source/Editor/TurnkeySupport/Private/TurnkeySupportModule.cpp b/Engine/Source/Editor/TurnkeySupport/Private/TurnkeySupportModule.cpp
--- a/Engine/Source/Editor/TurnkeySupport/Private/TurnkeySupportModule.cpp
+++ b/Engine/Source/Editor/TurnkeySupport/Private/TurnkeySupportModule.cpp
@@ -1453,9 +1453,16 @@ void FTurnkeySupportModule::MakeTurnkeyMenu(FToolMenuSection& MenuSection) const
 	ITargetDeviceServicesModule* TargetDeviceServicesModule = static_cast<ITargetDeviceServicesModule*>(FModuleManager::Get().LoadModule(TEXT("TargetDeviceServices")));
 	TargetDeviceServicesModule->GetDeviceProxyManager();
 
+	// hide during PIE
+	FUIAction PlatformMenuShownDelegate;
+	PlatformMenuShownDelegate.IsActionVisibleDelegate = FIsActionButtonVisible::CreateLambda([]() 
+	{
+		return !FTurnkeyEditorSupport::IsPIERunning();
+	});
+
  	MenuSection.AddEntry(FToolMenuEntry::InitComboButton(
 		"PlatformsMenu",
-		FUIAction(),
+		PlatformMenuShownDelegate,
 		FOnGetContent::CreateLambda([this] { return MakeTurnkeyMenuWidget(); }),
 		LOCTEXT("PlatformMenu", "Platforms"),
 		LOCTEXT("PlatformMenu_Tooltip", "Platform related actions and settings (Launching, Packaging, custom builds, etc)"),
