- Fixed Turnkey menu not respecting the platform configuration/target platform selections
- Improvements for ExecuteBuild - mostly around project selection, where it will give you options, cache last used project, etc
- Added search/replace modifier to the {....} syntax in the build spec (useful mainly for inivalue basically)
- Editor now just runs -command=ExecuteBuild -build=XYZ instead of parsing the XYZ build params and running UAT directly
- Deleted old #if 0'd code. Can go back in p4 history to find it later.
- Setup VisualStyles in Turnkey so that the project selector looks good on HiDPI

diff --git a/Engine/Source/Developer/DeveloperToolSettings/Classes/Settings/ProjectPackagingSettings.h b/Engine/Source/Developer/DeveloperToolSettings/Classes/Settings/ProjectPackagingSettings.h
--- a/Engine/Source/Developer/DeveloperToolSettings/Classes/Settings/ProjectPackagingSettings.h
+++ b/Engine/Source/Developer/DeveloperToolSettings/Classes/Settings/ProjectPackagingSettings.h
@@ -112,6 +112,8 @@ struct FProjectBuildSettings
 	 *   {iniif:Token:Config:Section:Key} - Replaced with Token if the vlaue for Key in Config's Section evaluates to True. Ex: {iniif:-iostore:/Script/UnrealEd.ProjectPackagingSettings:bUseIoStore}
 	 * Because ProjectPackagingSettings is a common section to read, if Config:Section: are not specified for 'iniif' or 'inivalue', it will use the ProjectPackagingSettings settings:
 	 *   {iniif:-iostore:bUseIoStore}
+	 * Additionally, the ini settings can have an optional search and replace modifier, to easily modify the string. The Replace can be blank:
+	 *   {inivalue:BuildConfiguration|PPBC_=} - This will get the BuildConfiguration from the settings, and then remove the PPBC_ enum prefix from the string, to just get say Development
 	 */
 	UPROPERTY(EditAnywhere, Category="Packaging")
 	FString BuildCookRunParams;
diff --git a/Engine/Source/Editor/TurnkeySupport/Private/TurnkeySupportModule.cpp b/Engine/Source/Editor/TurnkeySupport/Private/TurnkeySupportModule.cpp
--- a/Engine/Source/Editor/TurnkeySupport/Private/TurnkeySupportModule.cpp
+++ b/Engine/Source/Editor/TurnkeySupport/Private/TurnkeySupportModule.cpp
@@ -215,6 +215,7 @@ public:
 		if (PackagingSettings == nullptr)
 		{
 			PackagingSettings = NewObject<UProjectPackagingSettings>(GetTransientPackage());
+			// make sure any changes to DefaultGame are updated in this class
 			PackagingSettings->LoadSettingsForPlatform(PlatformString);
 		}
 
@@ -223,15 +224,17 @@ public:
 
 	static void CookOrPackage(FName IniPlatformName, EPrepareContentMode Mode)
 	{
-		const PlatformInfo::FTargetPlatformInfo* PlatformInfo = PlatformInfo::FindPlatformInfo(GetDefault<UProjectPackagingSettings>()->GetTargetPlatformForPlatform(IniPlatformName));
-		
+		// get all the helper objects
+		UProjectPackagingSettings* PackagingSettings = GetPackagingSettingsForPlatform(IniPlatformName);
+		// get a in-memory defaults which will have the user-settings, like the per-platform config/target platform stuff
+		UProjectPackagingSettings* AllPlatformPackagingSettings = GetMutableDefault<UProjectPackagingSettings>();
+	
+		// get the user-chosen targetplatform
+		const PlatformInfo::FTargetPlatformInfo* PlatformInfo = PlatformInfo::FindPlatformInfo(AllPlatformPackagingSettings->GetTargetPlatformForPlatform(IniPlatformName));
 		// this is unexpected to be able to happen, but it could if there was a bad value saved in the UProjectPackagingSettings - if this trips, we should handle errors
 		check(PlatformInfo != nullptr);
 
-
-		// get all the helper objects
 		const FString UBTPlatformString = PlatformInfo->DataDrivenPlatformInfo->UBTPlatformString;
-		UProjectPackagingSettings* PackagingSettings = GetPackagingSettingsForPlatform(IniPlatformName);
 		const FString ProjectPath = FPaths::IsProjectFilePathSet() ? FPaths::ConvertRelativePathToFull(FPaths::GetProjectFilePath()) : FPaths::RootDir() / FApp::GetProjectName() / FApp::GetProjectName() + TEXT(".uproject");
 
 
@@ -399,7 +402,8 @@ public:
 				BuildCookRunParams += FString::Printf(TEXT(" -manifests -createchunkinstall -chunkinstalldirectory=\"%s\" -chunkinstallversion=%s"), *(PackagingSettings->HttpChunkInstallDataDirectory.Path), *(PackagingSettings->HttpChunkInstallDataVersion));
 			}
 
-			const UProjectPackagingSettings::FConfigurationInfo& ConfigurationInfo = UProjectPackagingSettings::ConfigurationInfo[(int)PackagingSettings->GetBuildConfigurationForPlatform(IniPlatformName)];
+			// use AllPlatformPackagingSettings because these are user settings, and not checked in for all users to use
+			const UProjectPackagingSettings::FConfigurationInfo& ConfigurationInfo = UProjectPackagingSettings::ConfigurationInfo[(int)AllPlatformPackagingSettings->GetBuildConfigurationForPlatform(IniPlatformName)];
 			if (PlatformInfo->PlatformType == EBuildTargetType::Server)
 			{
 				BuildCookRunParams += FString::Printf(TEXT(" -serverconfig=%s"), LexToString(ConfigurationInfo.Configuration));
@@ -457,123 +461,13 @@ public:
 		return true;
 	}
 
-	static FString GetIniSetting(FString Spec)
-	{
-		// handle these cases:
-		//  iniif:-option:Engine:/Script/Module.Class:bUseOption
-		//  iniif:-option:bUseOption   [convenience for ProjectPackagingSettings setting]
-		//  inivalue:Engine:/Script/Module.Class:SomeSetting
-		//  inivalue:SomeSetting       [convenience for ProjectPackagingSettings setting]
-
-		TArray<FString> Tokens;
-		Spec.ParseIntoArray(Tokens, TEXT(":"));
-
-		FString ConfigName;
-		FString SectionName;
-		FString Key;
-		FString IniIfValue;
-		if (Tokens[0] == TEXT("iniif"))
-		{
-			if (Tokens.Num() == 3)
-			{
-				ConfigName = TEXT("Game");
-				SectionName = TEXT("/Script/UnrealEd.ProjectPackagingSettings");
-				Key = Tokens[2];
-				IniIfValue = Tokens[1];
-			}
-			else if (Tokens.Num() == 5)
-			{
-				ConfigName = Tokens[2];
-				SectionName = Tokens[3];
-				Key = Tokens[4];
-				IniIfValue = Tokens[1];
-			}
-			else
-			{
-				UE_LOG(LogTurnkeySupport, Warning, TEXT("Found a bad iniif spec: %s"), *Spec);
-				return TEXT("");
-			}
-		}
-		else if (Tokens[0] == TEXT("inivalue"))
-		{
-			if (Tokens.Num() == 2)
-			{
-				ConfigName = TEXT("Game");
-				SectionName = TEXT("/Script/UnrealEd.ProjectPackagingSettings");
-				Key = Tokens[1];
-			}
-			else if (Tokens.Num() == 4)
-			{
-				ConfigName = Tokens[1];
-				SectionName = Tokens[2];
-				Key = Tokens[3];
-			}
-			else
-			{
-				UE_LOG(LogTurnkeySupport, Warning, TEXT("Found a bad inivalue spec: %s"), *Spec);
-				return TEXT("");
-			}
-		}
-		else
-		{
-			UE_LOG(LogTurnkeySupport, Warning, TEXT("Found a bad ini spec: %s"), *Spec);
-			return TEXT("");
-		}
-
-		// get the value, if it exists (or empty string if not)
-		FString FoundValue;
-		GConfig->GetString(*SectionName, *Key, FoundValue, ConfigName);
-
-		if (Tokens[0] == TEXT("iniif"))
-		{
-			if (FCString::ToBool(*FoundValue))
-			{
-				return IniIfValue;
-			}
-			return TEXT("");
-		}
-
-		return FoundValue;
-
-	}
-
-	// @todo since Turnkey can execute builds with ExecuteBuild, we can probably just let it do all of the work so that we don't have the commandline fixup code in two places
 	static void ExecuteCustomBuild(FName IniPlatformName, FProjectBuildSettings Build)
 	{
 		const PlatformInfo::FTargetPlatformInfo* PlatformInfo = PlatformInfo::FindPlatformInfo(GetDefault<UProjectPackagingSettings>()->GetTargetPlatformForPlatform(IniPlatformName));
-
-		FString CommandLine = Build.BuildCookRunParams;
-
 		const FString ProjectPath = FPaths::IsProjectFilePathSet() ? FPaths::ConvertRelativePathToFull(FPaths::GetProjectFilePath()) : FPaths::RootDir() / FApp::GetProjectName() / FApp::GetProjectName() + TEXT(".uproject");
-		CommandLine = CommandLine.Replace(TEXT("{Project}"), *ProjectPath);
-
-		CommandLine = CommandLine.Replace(TEXT("{Platform}"), *IniPlatformName.ToString());
-
-		int Offset = 0;
-		while (Offset != -1)
-		{
-			Offset = CommandLine.Find(TEXT("{ini"), ESearchCase::IgnoreCase, ESearchDir::FromStart, Offset);
-			if (Offset != -1)
-			{
-				// skip over the {
-				Offset++;
-
-				// look for the close brace, don't need to worry about nested braces
-				int CloseBrace = CommandLine.Find(TEXT("}"), ESearchCase::IgnoreCase, ESearchDir::FromStart, Offset);
-				if (CloseBrace == -1)
-				{
-					Offset = -1;
-				}
-				else
-				{
-					FString IniSpec = CommandLine.Mid(Offset, CloseBrace - Offset);
-					FString ReplacedValue = GetIniSetting(IniSpec);
-					CommandLine = CommandLine.Mid(0, Offset - 1) + ReplacedValue + CommandLine.Mid(CloseBrace + 1);
-				}
-			}
-		}
-
 
+		FString CommandLine = FString::Printf(TEXT("-ScriptsForProject=\"%s\" Turnkey -command=ExecuteBuild -build=\"%s\" -platform=%s -project=\"%s\""),
+			*ProjectPath, *Build.Name, *IniPlatformName.ToString(), *ProjectPath);
 
 		FTurnkeyEditorSupport::RunUAT(CommandLine, PlatformInfo->DisplayName, LOCTEXT("Turnkey_CustomTaskName", "Executing Custom Build"), LOCTEXT("Turnkey_CustomTaskName", "Custom"), FEditorStyle::GetBrush(TEXT("MainFrame.PackageProject")));
 	}
@@ -1858,567 +1752,6 @@ void FTurnkeySupportModule::ShutdownModule( )
 }
 
 
+IMPLEMENT_MODULE(FTurnkeySupportModule, TurnkeySupport);
 
 #undef LOCTEXT_NAMESPACE
-
-
-
-
-
-
-
-
-
-
-#if 0
-
-// PACKAGING STUFF THAT HASN'T COME OVER
-
-
-GUnrealEd->CancelPlayingViaLauncher();
-SaveAll();
-
-// does the project have any code?
-FGameProjectGenerationModule& GameProjectModule = FModuleManager::LoadModuleChecked<FGameProjectGenerationModule>(TEXT("GameProjectGeneration"));
-bool bProjectHasCode = GameProjectModule.Get().ProjectHasCodeFiles();
-
-const PlatformInfo::FTargetPlatformInfo* const PlatformInfo = PlatformInfo::FindPlatformInfo(InPlatformInfoName);
-check(PlatformInfo);
-
-if (FInstalledPlatformInfo::Get().IsPlatformMissingRequiredFile(PlatformInfo->DataDrivenPlatformInfo->UBTPlatformString))
-{
-	if (!FInstalledPlatformInfo::OpenInstallerOptions())
-	{
-		FMessageDialog::Open(EAppMsgType::Ok, LOCTEXT("MissingPlatformFilesPackage", "Missing required files to package this platform."));
-	}
-	return;
-}
-
-if (UGameMapsSettings::GetGameDefaultMap().IsEmpty())
-{
-	FMessageDialog::Open(EAppMsgType::Ok, LOCTEXT("MissingGameDefaultMap", "No Game Default Map specified in Project Settings > Maps & Modes."));
-	return;
-}
-
-if (/*PlatformInfo->DataDrivenPlatformInfo->GetSdkStatus() != ETurnkeyPlatformSdkStatus::Valid ||*/ (bProjectHasCode && PlatformInfo->DataDrivenPlatformInfo->bUsesHostCompiler && !FSourceCodeNavigation::IsCompilerAvailable()))
-{
-	if (!ShowBadSDKDialog(PlatformInfo->IniPlatformName))
-	{
-		// 			IMainFrameModule& MainFrameModule = FModuleManager::GetModuleChecked<IMainFrameModule>(TEXT("MainFrame"));
-		// 			MainFrameModule.BroadcastMainFrameSDKNotInstalled(PlatformInfo->Name.ToString(), PlatformInfo->DataDrivenPlatformInfo->SDKTutorial);
-		TArray<FAnalyticsEventAttribute> ParamArray;
-		ParamArray.Add(FAnalyticsEventAttribute(TEXT("Time"), 0.0));
-		FEditorAnalytics::ReportEvent(TEXT("Editor.Package.Failed"), PlatformInfo->Name.ToString(), bProjectHasCode, EAnalyticsErrorCodes::SDKNotFound, ParamArray);
-		return;
-	}
-}
-
-UProjectPackagingSettings* PackagingSettings = Cast<UProjectPackagingSettings>(UProjectPackagingSettings::StaticClass()->GetDefaultObject());
-const UProjectPackagingSettings::FConfigurationInfo& ConfigurationInfo = UProjectPackagingSettings::ConfigurationInfo[(int)PackagingSettings->BuildConfiguration];
-
-// Note: Blueprint nativization has been removed as a feature.
-const bool bAssetNativizationEnabled = false;
-
-const ITargetPlatform* const Platform = GetTargetPlatformManager()->FindTargetPlatform(PlatformInfo->Name.ToString());
-{
-	if (Platform)
-	{
-		FString NotInstalledTutorialLink;
-		FString DocumentationLink;
-		FText CustomizedLogMessage;
-
-		int32 Result = Platform->CheckRequirements(bProjectHasCode, ConfigurationInfo.Configuration, bAssetNativizationEnabled, NotInstalledTutorialLink, DocumentationLink, CustomizedLogMessage);
-
-		// report to analytics
-		FEditorAnalytics::ReportBuildRequirementsFailure(TEXT("Editor.Package.Failed"), PlatformInfo->Name.ToString(), bProjectHasCode, Result);
-
-		// report to main frame
-		bool UnrecoverableError = false;
-
-		// report to message log
-		if ((Result & ETargetPlatformReadyStatus::SDKNotFound) != 0)
-		{
-			AddMessageLog(
-				LOCTEXT("SdkNotFoundMessage", "Software Development Kit (SDK) not found."),
-				CustomizedLogMessage.IsEmpty() ? FText::Format(LOCTEXT("SdkNotFoundMessageDetail", "Please install the SDK for the {0} target platform!"), Platform->DisplayName()) : CustomizedLogMessage,
-				NotInstalledTutorialLink,
-				DocumentationLink
-			);
-			UnrecoverableError = true;
-		}
-
-		if ((Result & ETargetPlatformReadyStatus::LicenseNotAccepted) != 0)
-		{
-			AddMessageLog(
-				LOCTEXT("LicenseNotAcceptedMessage", "License not accepted."),
-				CustomizedLogMessage.IsEmpty() ? LOCTEXT("LicenseNotAcceptedMessageDetail", "License must be accepted in project settings to deploy your app to the device.") : CustomizedLogMessage,
-				NotInstalledTutorialLink,
-				DocumentationLink
-			);
-
-			UnrecoverableError = true;
-		}
-
-		if ((Result & ETargetPlatformReadyStatus::ProvisionNotFound) != 0)
-		{
-			AddMessageLog(
-				LOCTEXT("ProvisionNotFoundMessage", "Provision not found."),
-				CustomizedLogMessage.IsEmpty() ? LOCTEXT("ProvisionNotFoundMessageDetail", "A provision is required for deploying your app to the device.") : CustomizedLogMessage,
-				NotInstalledTutorialLink,
-				DocumentationLink
-			);
-			UnrecoverableError = true;
-		}
-
-		if ((Result & ETargetPlatformReadyStatus::SigningKeyNotFound) != 0)
-		{
-			AddMessageLog(
-				LOCTEXT("SigningKeyNotFoundMessage", "Signing key not found."),
-				CustomizedLogMessage.IsEmpty() ? LOCTEXT("SigningKeyNotFoundMessageDetail", "The app could not be digitally signed, because the signing key is not configured.") : CustomizedLogMessage,
-				NotInstalledTutorialLink,
-				DocumentationLink
-			);
-			UnrecoverableError = true;
-		}
-
-		if ((Result & ETargetPlatformReadyStatus::ManifestNotFound) != 0)
-		{
-			AddMessageLog(
-				LOCTEXT("ManifestNotFound", "Manifest not found."),
-				CustomizedLogMessage.IsEmpty() ? LOCTEXT("ManifestNotFoundMessageDetail", "The generated application manifest could not be found.") : CustomizedLogMessage,
-				NotInstalledTutorialLink,
-				DocumentationLink
-			);
-			UnrecoverableError = true;
-		}
-
-		if ((Result & ETargetPlatformReadyStatus::RemoveServerNameEmpty) != 0
-			&& (bProjectHasCode || (Result & ETargetPlatformReadyStatus::CodeBuildRequired)
-				|| (!FApp::GetEngineIsPromotedBuild() && !FApp::IsEngineInstalled())))
-		{
-			AddMessageLog(
-				LOCTEXT("RemoveServerNameNotFound", "Remote compiling requires a server name. "),
-				CustomizedLogMessage.IsEmpty() ? LOCTEXT("RemoveServerNameNotFoundDetail", "Please specify one in the Remote Server Name settings field.") : CustomizedLogMessage,
-				NotInstalledTutorialLink,
-				DocumentationLink
-			);
-			UnrecoverableError = true;
-		}
-
-		if ((Result & ETargetPlatformReadyStatus::CodeUnsupported) != 0)
-		{
-			FMessageDialog::Open(EAppMsgType::Ok, LOCTEXT("NotSupported_SelectedPlatform", "Sorry, packaging a code-based project for the selected platform is currently not supported. This feature may be available in a future release."));
-			UnrecoverableError = true;
-		}
-		else if ((Result & ETargetPlatformReadyStatus::PluginsUnsupported) != 0)
-		{
-			FMessageDialog::Open(EAppMsgType::Ok, LOCTEXT("NotSupported_ThirdPartyPlugins", "Sorry, packaging a project with third-party plugins is currently not supported for the selected platform. This feature may be available in a future release."));
-			UnrecoverableError = true;
-		}
-
-		if (UnrecoverableError)
-		{
-			return;
-		}
-	}
-}
-
-if (!FModuleManager::LoadModuleChecked<IProjectTargetPlatformEditorModule>("ProjectTargetPlatformEditor").ShowUnsupportedTargetWarning(PlatformInfo->VanillaInfo->Name))
-{
-	return;
-}
-
-
-
-void FInternalPlayWorldCommandCallbacks::HandleShowSDKTutorial(FString PlatformName, FString NotInstalledDocLink)
-{
-	// broadcast this, and assume someone will pick it up
-	IMainFrameModule& MainFrameModule = FModuleManager::GetModuleChecked<IMainFrameModule>(TEXT("MainFrame"));
-	MainFrameModule.BroadcastMainFrameSDKNotInstalled(PlatformName, NotInstalledDocLink);
-}
-
-
-
-#endif
-
-
-// another round of junk, some duplicates with above
-#if 0
-
-
-
-/////// Will move these into TurnkeySupportModule soon:
-
-
-
-bool FInternalPlayWorldCommandCallbacks::CanLaunchOnDevice(const FString& DeviceName)
-{
-	if (!GUnrealEd->IsPlayingViaLauncher())
-	{
-		static TWeakPtr<ITargetDeviceProxyManager> DeviceProxyManagerPtr;
-
-		if (!DeviceProxyManagerPtr.IsValid())
-		{
-			ITargetDeviceServicesModule* TargetDeviceServicesModule = FModuleManager::Get().LoadModulePtr<ITargetDeviceServicesModule>(TEXT("TargetDeviceServices"));
-			if (TargetDeviceServicesModule)
-			{
-				DeviceProxyManagerPtr = TargetDeviceServicesModule->GetDeviceProxyManager();
-			}
-		}
-
-		TSharedPtr<ITargetDeviceProxyManager> DeviceProxyManager = DeviceProxyManagerPtr.Pin();
-		if (DeviceProxyManager.IsValid())
-		{
-			TSharedPtr<ITargetDeviceProxy> DeviceProxy = DeviceProxyManager->FindProxy(DeviceName);
-			if (DeviceProxy.IsValid() && DeviceProxy->IsConnected() && DeviceProxy->IsAuthorized())
-			{
-				return true;
-			}
-
-			// check if this is an aggregate proxy
-			TArray<TSharedPtr<ITargetDeviceProxy>> Devices;
-			DeviceProxyManager->GetProxies(FName(*DeviceName), false, Devices);
-
-			// returns true if the game can be launched al least on 1 device
-			for (auto DevicesIt = Devices.CreateIterator(); DevicesIt; ++DevicesIt)
-			{
-				TSharedPtr<ITargetDeviceProxy> DeviceAggregateProxy = *DevicesIt;
-				if (DeviceAggregateProxy.IsValid() && DeviceAggregateProxy->IsConnected() && DeviceAggregateProxy->IsAuthorized())
-				{
-					return true;
-				}
-			}
-
-		}
-	}
-
-	return false;
-}
-
-
-void FInternalPlayWorldCommandCallbacks::LaunchOnDevice(const FString& DeviceId, const FString& DeviceName, bool bUseTurnkey)
-{
-	FTargetDeviceId TargetDeviceId;
-	if (FTargetDeviceId::Parse(DeviceId, TargetDeviceId))
-	{
-		const PlatformInfo::FTargetPlatformInfo* const PlatformInfo = PlatformInfo::FindPlatformInfo(*TargetDeviceId.GetPlatformName());
-		FString UBTPlatformName = PlatformInfo->DataDrivenPlatformInfo->UBTPlatformString;
-		FString IniPlatformName = PlatformInfo->IniPlatformName.ToString();
-
-		check(PlatformInfo);
-
-		if (FInstalledPlatformInfo::Get().IsPlatformMissingRequiredFile(UBTPlatformName))
-		{
-			if (!FInstalledPlatformInfo::OpenInstallerOptions())
-			{
-				FMessageDialog::Open(EAppMsgType::Ok, LOCTEXT("MissingPlatformFilesLaunch", "Missing required files to launch on this platform."));
-			}
-			return;
-		}
-
-		if (FModuleManager::LoadModuleChecked<IProjectTargetPlatformEditorModule>("ProjectTargetPlatformEditor").ShowUnsupportedTargetWarning(*TargetDeviceId.GetPlatformName()))
-		{
-			GUnrealEd->CancelPlayingViaLauncher();
-
-			FRequestPlaySessionParams::FLauncherDeviceInfo DeviceInfo;
-			DeviceInfo.DeviceId = DeviceId;
-			DeviceInfo.DeviceName = DeviceName;
-			// @todo turnkey: we set this to false because we will kick off a Turnkey run before cooking, etc, to get an early warning. however, if it's too difficult
-			// to get an error back from CreateUatTask, then we should set this to bUseTurnkey and remove the block below, and let the code in FLauncherWorker::CreateAndExecuteTasks handle it
-			DeviceInfo.bUpdateDeviceFlash = false;
-
-			FRequestPlaySessionParams SessionParams;
-			SessionParams.SessionDestination = EPlaySessionDestinationType::Launcher;
-			SessionParams.LauncherTargetDevice = DeviceInfo;
-
-			// if we want to check device flash before we start cooking, kick it off now. we could delay this 
-			if (bUseTurnkey)
-			{
-				FString CommandLine = FString::Printf(TEXT("Turnkey -command=VerifySdk -UpdateIfNeeded -platform=%s -EditorIO -noturnkeyvariables -device=%s -utf8output -WaitForUATMutex"), *UBTPlatformName, *DeviceName);
-				FText TaskName = LOCTEXT("VerifyingSDK", "Verifying SDK and Device");
-
-				IUATHelperModule::Get().CreateUatTask(CommandLine, FText::FromString(IniPlatformName), TaskName, TaskName, FEditorStyle::GetBrush(TEXT("MainFrame.PackageProject")),
-					[SessionParams](FString Result, double)
-					{
-						// unfortunate string comparison for success
-						bool bWasSuccessful = Result == TEXT("Completed");
-						AsyncTask(ENamedThreads::GameThread, [SessionParams, bWasSuccessful]()
-							{
-								if (bWasSuccessful)
-								{
-									GUnrealEd->RequestPlaySession(SessionParams);
-								}
-								else
-								{
-									TSharedRef<SWindow> Win = OpenMsgDlgInt_NonModal(EAppMsgType::YesNo, LOCTEXT("SDKCheckFailed", "SDK Verification failed. Would you like to attempt the Launch On anyway?"), LOCTEXT("SDKCheckFailedTitle", "SDK Verification"),
-										FOnMsgDlgResult::CreateLambda([SessionParams](const TSharedRef<SWindow>&, EAppReturnType::Type Choice)
-											{
-												if (Choice == EAppReturnType::Yes)
-												{
-													GUnrealEd->RequestPlaySession(SessionParams);
-												}
-											}));
-									Win->ShowWindow();
-								}
-							});
-					});
-			}
-			else
-			{
-				GUnrealEd->RequestPlaySession(SessionParams);
-			}
-		}
-	}
-}
-
-
-void FInternalPlayWorldCommandCallbacks::RepeatLastLaunch_Clicked()
-{
-	const ULevelEditorPlaySettings* PlaySettings = GetDefault<ULevelEditorPlaySettings>();
-
-	switch (PlaySettings->LastExecutedLaunchModeType)
-	{
-	case LaunchMode_OnDevice:
-		if (FParse::Param(FCommandLine::Get(), TEXT("OldLaunchOn")))
-		{
-			if (IsReadyToLaunchOnDevice(PlaySettings->LastExecutedLaunchDevice))
-			{
-				LaunchOnDevice(PlaySettings->LastExecutedLaunchDevice, PlaySettings->LastExecutedLaunchName, false);
-			}
-		}
-		else
-		{
-			// @todo turnkey - remove this function completely and all the platform support, and put it into Turnkey
-//			if (IsReadyToLaunchOnDevice(PlaySettings->LastExecutedLaunchDevice))
-			{
-				LaunchOnDevice(PlaySettings->LastExecutedLaunchDevice, PlaySettings->LastExecutedLaunchName, true);
-			}
-		}
-
-
-		break;
-
-	default:
-		break;
-	}
-}
-
-
-bool FInternalPlayWorldCommandCallbacks::RepeatLastLaunch_CanExecute()
-{
-	const ULevelEditorPlaySettings* PlaySettings = GetDefault<ULevelEditorPlaySettings>();
-
-	switch (PlaySettings->LastExecutedLaunchModeType)
-	{
-	case LaunchMode_OnDevice:
-		return CanLaunchOnDevice(PlaySettings->LastExecutedLaunchName);
-
-	default:
-		return false;
-	}
-}
-
-
-FText FInternalPlayWorldCommandCallbacks::GetRepeatLastLaunchToolTip()
-{
-	const ULevelEditorPlaySettings* PlaySettings = GetDefault<ULevelEditorPlaySettings>();
-
-	switch (PlaySettings->LastExecutedLaunchModeType)
-	{
-	case LaunchMode_OnDevice:
-		if (CanLaunchOnDevice(PlaySettings->LastExecutedLaunchName))
-		{
-			FFormatNamedArguments Arguments;
-			Arguments.Add(TEXT("DeviceName"), FText::FromString(PlaySettings->LastExecutedLaunchName));
-
-			return FText::Format(LOCTEXT("RepeatLaunchTooltip", "Launch this level on {DeviceName}"), Arguments);
-		}
-
-		break;
-
-	default:
-		break;
-	}
-
-	return LOCTEXT("RepeatLaunchSelectOptionToolTip", "Select a play-on target from the combo menu");
-}
-
-
-FSlateIcon FInternalPlayWorldCommandCallbacks::GetRepeatLastLaunchIcon()
-{
-	const ULevelEditorPlaySettings* PlaySettings = GetDefault<ULevelEditorPlaySettings>();
-
-	// @todo gmp: add play mode specific icons
-	switch (PlaySettings->LastExecutedLaunchModeType)
-	{
-	case LaunchMode_OnDevice:
-		break;
-	}
-
-	static FName RepeatLastLaunchIcon("PlayWorld.RepeatLastLaunch");
-
-	return FSlateIcon(FEditorStyle::GetStyleSetName(), RepeatLastLaunchIcon);
-}
-
-bool FInternalPlayWorldCommandCallbacks::IsReadyToLaunchOnDevice(FString DeviceId)
-{
-	int32 Index = 0;
-	DeviceId.FindChar(TEXT('@'), Index);
-	FName PlatformName = FName(*DeviceId.Left(Index));
-
-	const PlatformInfo::FTargetPlatformInfo* PlatformInfo = PlatformInfo::FindPlatformInfo(PlatformName);
-	checkf(PlatformInfo, TEXT("Unable to find PlatformInfo for %s"), *PlatformName.ToString());
-
-	FGameProjectGenerationModule& GameProjectModule = FModuleManager::LoadModuleChecked<FGameProjectGenerationModule>(TEXT("GameProjectGeneration"));
-	bool bHasCode = GameProjectModule.Get().ProjectHasCodeFiles();
-
-	if (ITurnkeySupportModule::Get().GetSdkInfo(PlatformName).Status != ETurnkeyPlatformSdkStatus::Valid)
-	{
-		IMainFrameModule& MainFrameModule = FModuleManager::GetModuleChecked<IMainFrameModule>(TEXT("MainFrame"));
-		MainFrameModule.BroadcastMainFrameSDKNotInstalled(PlatformInfo->Name.ToString(), PlatformInfo->DataDrivenPlatformInfo->SDKTutorial);
-		TArray<FAnalyticsEventAttribute> ParamArray;
-		ParamArray.Add(FAnalyticsEventAttribute(TEXT("Time"), 0.0));
-		FEditorAnalytics::ReportEvent(TEXT("Editor.LaunchOn.Failed"), PlatformInfo->Name.ToString(), bHasCode, EAnalyticsErrorCodes::SDKNotFound, ParamArray);
-		return false;
-	}
-
-	const ITargetPlatform* Platform = GetTargetPlatformManager()->FindTargetPlatform(PlatformName);
-	if (Platform)
-	{
-		FString NotInstalledTutorialLink;
-		FString DocumentationLink;
-		FText CustomizedLogMessage;
-
-		EBuildConfiguration BuildConfiguration = GetDefault<ULevelEditorPlaySettings>()->GetLaunchBuildConfiguration();
-		bool bEnableAssetNativization = false;
-		int32 Result = Platform->CheckRequirements(bHasCode, BuildConfiguration, bEnableAssetNativization, NotInstalledTutorialLink, DocumentationLink, CustomizedLogMessage);
-
-		// report to analytics
-		FEditorAnalytics::ReportBuildRequirementsFailure(TEXT("Editor.LaunchOn.Failed"), PlatformName.ToString(), bHasCode, Result);
-
-		// report to message log
-		bool UnrecoverableError = false;
-
-		if ((Result & ETargetPlatformReadyStatus::SDKNotFound) != 0)
-		{
-			AddMessageLog(
-				LOCTEXT("SdkNotFoundMessage", "Software Development Kit (SDK) not found."),
-				CustomizedLogMessage.IsEmpty() ? FText::Format(LOCTEXT("SdkNotFoundMessageDetail", "Please install the SDK for the {0} target platform!"), Platform->DisplayName()) : CustomizedLogMessage,
-				NotInstalledTutorialLink,
-				DocumentationLink
-			);
-
-			UnrecoverableError = true;
-		}
-
-		if ((Result & ETargetPlatformReadyStatus::LicenseNotAccepted) != 0)
-		{
-			AddMessageLog(
-				LOCTEXT("LicenseNotAcceptedMessage", "License not accepted."),
-				CustomizedLogMessage.IsEmpty() ? LOCTEXT("LicenseNotAcceptedMessageDetail", "License must be accepted in project settings to deploy your app to the device.") : CustomizedLogMessage,
-				NotInstalledTutorialLink,
-				DocumentationLink
-			);
-
-			UnrecoverableError = true;
-		}
-
-		if ((Result & ETargetPlatformReadyStatus::ProvisionNotFound) != 0)
-		{
-			AddMessageLog(
-				LOCTEXT("ProvisionNotFoundMessage", "Provision not found."),
-				CustomizedLogMessage.IsEmpty() ? LOCTEXT("ProvisionNotFoundMessageDetail", "A provision is required for deploying your app to the device.") : CustomizedLogMessage,
-				NotInstalledTutorialLink,
-				DocumentationLink
-			);
-
-			UnrecoverableError = true;
-		}
-
-		if ((Result & ETargetPlatformReadyStatus::SigningKeyNotFound) != 0)
-		{
-			AddMessageLog(
-				LOCTEXT("SigningKeyNotFoundMessage", "Signing key not found."),
-				CustomizedLogMessage.IsEmpty() ? LOCTEXT("SigningKeyNotFoundMessageDetail", "The app could not be digitally signed, because the signing key is not configured.") : CustomizedLogMessage,
-				NotInstalledTutorialLink,
-				DocumentationLink
-			);
-
-			UnrecoverableError = true;
-		}
-
-		if ((Result & ETargetPlatformReadyStatus::ManifestNotFound) != 0)
-		{
-			AddMessageLog(
-				LOCTEXT("ManifestNotFound", "Manifest not found."),
-				CustomizedLogMessage.IsEmpty() ? LOCTEXT("ManifestNotFoundMessageDetail", "The generated application manifest could not be found.") : CustomizedLogMessage,
-				NotInstalledTutorialLink,
-				DocumentationLink
-			);
-
-			UnrecoverableError = true;
-		}
-
-		if ((Result & ETargetPlatformReadyStatus::RemoveServerNameEmpty) != 0
-			&& (bHasCode || (Result & ETargetPlatformReadyStatus::CodeBuildRequired)
-				|| (!FApp::GetEngineIsPromotedBuild() && !FApp::IsEngineInstalled())))
-		{
-			AddMessageLog(
-				LOCTEXT("RemoveServerNameNotFound", "Remote compiling requires a server name. "),
-				CustomizedLogMessage.IsEmpty() ? LOCTEXT("RemoveServerNameNotFoundDetail", "Please specify one in the Remote Server Name settings field.") : CustomizedLogMessage,
-				NotInstalledTutorialLink,
-				DocumentationLink
-			);
-			UnrecoverableError = true;
-		}
-
-		if (UnrecoverableError)
-		{
-			return false;
-		}
-
-		// report to main frame
-		if ((Result & ETargetPlatformReadyStatus::CodeUnsupported) != 0)
-		{
-			// show the message
-			FMessageDialog::Open(EAppMsgType::Ok, LOCTEXT("NotSupported_CodeBased", "Sorry, launching a code-based project for the selected platform is currently not supported. This feature may be available in a future release."));
-			return false;
-		}
-		if ((Result & ETargetPlatformReadyStatus::PluginsUnsupported) != 0)
-		{
-			// show the message
-			FMessageDialog::Open(EAppMsgType::Ok, LOCTEXT("NotSupported_Plugins", "Sorry, launching a project with third-party plugins is currently not supported for the selected platform. This feature may be available in a future release."));
-			return false;
-		}
-	}
-	else
-	{
-		IMainFrameModule& MainFrameModule = FModuleManager::GetModuleChecked<IMainFrameModule>(TEXT("MainFrame"));
-		MainFrameModule.BroadcastMainFrameSDKNotInstalled(PlatformInfo->Name.ToString(), PlatformInfo->DataDrivenPlatformInfo->SDKTutorial);
-		return false;
-	}
-
-	return true;
-}
-
-
-
-#endif
-
-
-
-
-
-
-
-
-
-
-
-
-
-
-
-
-
-IMPLEMENT_MODULE(FTurnkeySupportModule, TurnkeySupport);
diff --git a/Engine/Source/Programs/AutomationTool/Turnkey/Commands/ExecuteBuild.cs b/Engine/Source/Programs/AutomationTool/Turnkey/Commands/ExecuteBuild.cs
--- a/Engine/Source/Programs/AutomationTool/Turnkey/Commands/ExecuteBuild.cs
+++ b/Engine/Source/Programs/AutomationTool/Turnkey/Commands/ExecuteBuild.cs
@@ -4,7 +4,7 @@ using AutomationTool;
 using System;
 using System.Collections.Generic;
 using System.Linq;
-using System.Text;
+using System.IO;
 using System.Text.RegularExpressions;
 using EpicGames.Core;
 using UnrealBuildTool;
@@ -13,59 +13,94 @@ namespace Turnkey.Commands
 {
 	class ExecuteBuild : TurnkeyCommand
 	{
+		// could move this to TurnkeyUtils!
+		static string GetStructEntry(string Input, string Property, bool bIsArrayProperty)
+		{
+			string PrimaryRegex;
+			string AltRegex = null;
+			if (bIsArrayProperty)
+			{
+				PrimaryRegex = string.Format("{0}\\s*=\\s*\\((.*?)\\)", Property);
+			}
+			else
+			{
+				// handle quoted strings, allowing for quoited quotation marks (basically doing " followed by whatever, until we see a quote that was not proceeded by a \, and gather the whole mess in an outer group)
+				PrimaryRegex = string.Format("{0}\\s*=\\s*\"((.*?)[^\\\\])\"", Property);
+				AltRegex = string.Format("{0}\\s*=\\s*(.*?)\\,", Property);
+			}
+
+			// attempt to match it!
+			Match Result = Regex.Match(Input, PrimaryRegex);
+			if (!Result.Success)
+			{
+				Result = Regex.Match(Input, AltRegex);
+			}
+
+			// if we got a success, return the main match value
+			if (Result.Success)
+			{
+				return Result.Groups[1].Value.ToString();
+			}
+
+			return null;
+		}
+
 		protected override void Execute(string[] CommandOptions)
 		{
 			// we need a platform to execute
 			List<UnrealTargetPlatform> Platforms = TurnkeyUtils.GetPlatformsFromCommandLineOrUser(CommandOptions, null);
+			FileReference ProjectFile = TurnkeyUtils.GetProjectFromCommandLineOrUser(CommandOptions);
 
-			string Project = TurnkeyUtils.GetVariableValue("Project");
-			if (string.IsNullOrEmpty(Project))
+			// we need a project file, so if canceled, abore this command
+			if (ProjectFile == null)
 			{
-				Project = TurnkeyUtils.ReadInput("Enter a project to build");
+				return;
 			}
-			
-			FileReference ProjectFile = ProjectUtils.FindProjectFileFromName(Project);
 
 			string DesiredBuild = TurnkeyUtils.ParseParamValue("Build", null, CommandOptions);
 
 			// get a list of builds from config
 			foreach (UnrealTargetPlatform Platform in Platforms)
 			{
- 				ConfigHierarchy GameConfig = ConfigCache.ReadHierarchy(ConfigHierarchyType.Game, ProjectFile.Directory, Platform);
+				ConfigHierarchy GameConfig = ConfigCache.ReadHierarchy(ConfigHierarchyType.Game, ProjectFile.Directory, Platform);
 
 				List<string> Builds;
 				GameConfig.GetArray("/Script/UnrealEd.ProjectPackagingSettings", "ExtraProjectBuilds", out Builds);
 
 				Dictionary<string, string> BuildCommands = new Dictionary<string, string>(StringComparer.InvariantCultureIgnoreCase);
-				foreach (string Build in Builds)
+				if (Builds != null)
 				{
-					Match NameResult = Regex.Match(Build, "Name\\s*=\\s*\"(.*?)\"");
-					Match PlatformsResult = Regex.Match(Build, "SpecificPlatforms\\s*=\\s*(.*?)");
-					Match ParamsResult = Regex.Match(Build, "BuildCookRunParams=\\s*\"(.*?)\"");
-
-					// make sure required entries are there
-					if (!NameResult.Success || !ParamsResult.Success)
+					foreach (string Build in Builds)
 					{
-						continue;
-					}
+						string Name = GetStructEntry(Build, "Name", false);
+						string SpecificPlatforms = GetStructEntry(Build, "Platforms", true);
+						string Params = GetStructEntry(Build, "BuildCookRunParams", false);
 
-					// if platforms are specified, and this platform isn't one of them, skip it
-					if (PlatformsResult.Success)
-					{
-						string[] SpecificPlatforms = PlatformsResult.Groups[1].Value.Split(",\"".ToCharArray());
-						if (SpecificPlatforms.Length > 0 && !SpecificPlatforms.Contains(Platform.ToString()))
+						// make sure required entries are there
+						if (Name == null || Params == null)
 						{
 							continue;
 						}
-					}
 
-					// add to list of commands
-					BuildCommands.Add(NameResult.Groups[1].Value, ParamsResult.Groups[1].Value);
+						// if platforms are specified, and this platform isn't one of them, skip it
+						if (!string.IsNullOrEmpty(SpecificPlatforms))
+						{
+							string[] PlatformList = SpecificPlatforms.Split(",\"".ToCharArray(), StringSplitOptions.RemoveEmptyEntries);
+							// case insensitive Contains
+							if (PlatformList.Length > 0 && !PlatformList.Any(x => x.Equals(Platform.ToString(), StringComparison.OrdinalIgnoreCase)))
+							{
+								continue;
+							}
+						}
+
+						// add to list of commands
+						BuildCommands.Add(Name, Params);
+					}
 				}
 
 				if (BuildCommands.Count == 0)
 				{
-					TurnkeyUtils.Log("Unable to find a build for platform {0} and project {1}", Platform, Project);
+					TurnkeyUtils.Log("Unable to find a build for platform {0} and project {1}", Platform, ProjectFile.GetFileNameWithoutAnyExtensions());
 					continue;
 				}
 
@@ -120,7 +155,8 @@ namespace Turnkey.Commands
 			//  inivalue:Engine:/Script/Module.Class:SomeSetting
 			//  inivalue:SomeSetting       [convenience for ProjectPackagingSettings setting]
 
-			string[] Tokens = Spec.Split(":".ToCharArray());
+			string[] CommandAndModifiers = Spec.Split("|".ToCharArray(), StringSplitOptions.RemoveEmptyEntries);
+			string[] Tokens = CommandAndModifiers[0].Split(":".ToCharArray());
 
 			string ConfigName;
 			string SectionName;
@@ -185,14 +221,30 @@ namespace Turnkey.Commands
 			string FoundValue;
 			Config.GetString(SectionName, Key, out FoundValue);
 
+
 			if (Tokens[0].Equals("iniif", StringComparison.InvariantCultureIgnoreCase))
 			{
 				bool bIsTrue;
 				if (bool.TryParse(FoundValue, out bIsTrue) && bIsTrue)
 				{
-					return IniIfValue;
+					FoundValue = IniIfValue;
 				}
-				return "";
+				else
+				{
+					return "";
+				}
+			}
+
+			// look to see if we have a replace modifier to update the ini value, and apply it if so
+			if (CommandAndModifiers.Length > 1)
+			{
+				string[] SearchAndReplace = CommandAndModifiers[1].Split("=".ToCharArray());
+				if (SearchAndReplace.Length != 2)
+				{
+					TurnkeyUtils.Log("Found a search/replace modifier {0} in spec {1}", CommandAndModifiers, Spec);
+					return "";
+				}
+				FoundValue = FoundValue.Replace(SearchAndReplace[0], SearchAndReplace[1]);
 			}
 
 			return FoundValue;
diff --git a/Engine/Source/Programs/AutomationTool/Turnkey/Providers/HybridIOProvider.cs b/Engine/Source/Programs/AutomationTool/Turnkey/Providers/HybridIOProvider.cs
--- a/Engine/Source/Programs/AutomationTool/Turnkey/Providers/HybridIOProvider.cs
+++ b/Engine/Source/Programs/AutomationTool/Turnkey/Providers/HybridIOProvider.cs
@@ -17,22 +17,8 @@ namespace Turnkey
 {
 	class HybridIOProvider : ConsoleIOProvider
 	{
-		// @todo turnkey: if this is really bad on Mono, then we can set it in the manifest, but that means changing AutomationTool's and ATLaunchers' manifests
-		[System.Runtime.InteropServices.DllImport("user32.dll")]
-		private static extern bool SetProcessDPIAware();
 		static HybridIOProvider()
 		{
-			// make the form look good on modern displays!
-			if (!UnrealBuildTool.Utils.IsRunningOnMono && Environment.OSVersion.Version.Major >= 6)
-			{
-				SetProcessDPIAware();
-			}
-
-#if WINDOWS
-			Application.EnableVisualStyles();
-			Application.SetCompatibleTextRenderingDefault(false);
-#endif
-
 		}
 
 		private string ShowMacDialog(string Prompt, string Default)
@@ -60,48 +46,102 @@ namespace Turnkey
 #if WINDOWS
 		private string ShowFormsDialogs(string Prompt, string Default)
 		{
-			// TODO: Make sure this dialog actually looks like we expect
-			Size size = new Size(200, 70);
-			Form inputBox = new Form();
-
-			inputBox.FormBorderStyle = FormBorderStyle.FixedDialog;
-			inputBox.ClientSize = size;
-			inputBox.Text = "Turnkey Input";
-
-			TextBox textBox = new TextBox();
-			textBox.Size = new Size(size.Width - 10, 23);
-			textBox.Location = new Point(5, 5);
-			textBox.Text = Prompt;
-			inputBox.Controls.Add(textBox);
-
-			Button okButton = new Button();
-			okButton.DialogResult = DialogResult.OK;
-			okButton.Name = "okButton";
-			okButton.Size = new Size(75, 23);
-			okButton.Text = "&OK";
-			okButton.Location = new Point(size.Width - 80 - 80, 39);
-			inputBox.Controls.Add(okButton);
-
-			Button cancelButton = new Button();
-			cancelButton.DialogResult = DialogResult.Cancel;
-			cancelButton.Name = "cancelButton";
-			cancelButton.Size = new Size(75, 23);
-			cancelButton.Text = "&Cancel";
-			cancelButton.Location = new Point(size.Width - 80, 39);
-			inputBox.Controls.Add(cancelButton);
-
-			inputBox.AcceptButton = okButton;
-			inputBox.CancelButton = cancelButton;
-
-			DialogResult result = inputBox.ShowDialog();
-			if (result == DialogResult.OK)
-			{
-				return textBox.Text;
-			}
-			else
+			// Create a new instance of the form.
+			Form Form1 = new Form();
+
+			TextBox Label;
+			Button CancelBtn;
+			Button OkBtn;
+			TextBox TextBox;
+
+			TextBox = new TextBox();
+			Label = new TextBox();
+			CancelBtn = new Button();
+			OkBtn = new Button();
+			Form1.SuspendLayout();
+			// 
+			// Label
+			// 
+			Label.Anchor = ((AnchorStyles)(((AnchorStyles.Top | AnchorStyles.Left)
+			| AnchorStyles.Right | AnchorStyles.Bottom)));
+			Label.BackColor = System.Drawing.SystemColors.Control;
+			Label.BorderStyle = BorderStyle.None;
+			Label.Location = new System.Drawing.Point(13, 13);
+			Label.Multiline = true;
+			Label.Name = "Label";
+			Label.ScrollBars = ScrollBars.Vertical;
+			Label.ShortcutsEnabled = false;
+			Label.Size = new System.Drawing.Size(775, 199);
+			Label.TabIndex = 0;
+			Label.TabStop = false;
+			Label.Text = Prompt;
+
+			// 
+			// CancelBtn
+			// 
+			CancelBtn.Anchor = ((AnchorStyles)((AnchorStyles.Right | AnchorStyles.Bottom)));
+			CancelBtn.DialogResult = DialogResult.Cancel;
+			CancelBtn.Location = new System.Drawing.Point(679, 262);
+			CancelBtn.Margin = new Padding(4);
+			CancelBtn.Name = "CancelBtn";
+			CancelBtn.Size = new System.Drawing.Size(109, 34);
+			CancelBtn.TabIndex = 6;
+			CancelBtn.Text = "Cancel";
+			CancelBtn.UseVisualStyleBackColor = true;
+			// 
+			// TextBox
+			// 
+			TextBox.Anchor = ((AnchorStyles)(((AnchorStyles.Left | AnchorStyles.Bottom)
+			| AnchorStyles.Right)));
+			TextBox.Location = new System.Drawing.Point(13, 219);
+			TextBox.Margin = new Padding(4);
+			TextBox.Name = "TextBox";
+			TextBox.Size = new System.Drawing.Size(772, 22);
+			TextBox.TabIndex = 4;
+			TextBox.Text = Default;
+			// 
+			// OkBtn
+			// 
+			OkBtn.Anchor = ((AnchorStyles)((AnchorStyles.Right | AnchorStyles.Bottom)));
+			OkBtn.DialogResult = DialogResult.OK;
+			OkBtn.Location = new System.Drawing.Point(562, 262);
+			OkBtn.Margin = new Padding(4);
+			OkBtn.Name = "OkBtn";
+			OkBtn.Size = new System.Drawing.Size(109, 34);
+			OkBtn.TabIndex = 5;
+			OkBtn.Text = "Ok";
+			OkBtn.UseVisualStyleBackColor = true;
+			// 
+			// Form1
+			// 
+			Form1.AcceptButton = OkBtn;
+			Form1.StartPosition = FormStartPosition.CenterScreen;
+			Form1.AutoScaleDimensions = new System.Drawing.SizeF(8F, 16F);
+			Form1.AutoScaleMode = AutoScaleMode.Font;
+			Form1.CancelButton = CancelBtn;
+			Form1.ClientSize = new System.Drawing.Size(800, 315);
+			Form1.Controls.Add(CancelBtn);
+			Form1.Controls.Add(OkBtn);
+			Form1.Controls.Add(TextBox);
+			Form1.Controls.Add(Label);
+			Form1.Name = "Form1";
+			Form1.Text = "Make selection";
+			Form1.ResumeLayout(false);
+			Form1.PerformLayout();
+
+			// Display the form as a modal dialog box.
+			Form1.ShowDialog();
+
+			string Result = null;
+
+			// Determine if the OK button was clicked on the dialog box.
+			if (Form1.DialogResult == DialogResult.OK)
 			{
-				return Default;
+				Result = TextBox.Text;
 			}
+
+			Form1.Dispose();
+			return Result;
 		}
 #endif
 
diff --git a/Engine/Source/Programs/AutomationTool/Turnkey/Turnkey.Automation.cs b/Engine/Source/Programs/AutomationTool/Turnkey/Turnkey.Automation.cs
--- a/Engine/Source/Programs/AutomationTool/Turnkey/Turnkey.Automation.cs
+++ b/Engine/Source/Programs/AutomationTool/Turnkey/Turnkey.Automation.cs
@@ -13,33 +13,14 @@ namespace Turnkey
 	{
 		public override ExitCode Execute()
 		{
-			// 			Thread t = new Thread(Turnkey.ThreadProc);
-			// 			Console.WriteLine("Before setting apartment state: {0}",
-			// 				t.GetApartmentState());
-			// 
-			// 			t.SetApartmentState(ApartmentState.STA);
-			// 			Console.WriteLine("After setting apartment state: {0}",
-			// 				t.GetApartmentState());
-			// 
-			// 			t.Start(this);
-			// 			t.Join();
-
-			ThreadProc(this);
-			return TurnkeyUtils.ExitCode;
-		}
-
-		private static void ThreadProc(object Data)
-		{
-			BuildCommand Build = (BuildCommand)Data;
-
 			IOProvider IOProvider;
-			if (Build.ParseParam("EditorIO"))
+			if (ParseParam("EditorIO"))
 			{
 				IOProvider = new HybridIOProvider();
 			}
 			else
 			{
-				string ReportFilename = Build.ParseParamValue("ReportFilename");
+				string ReportFilename = ParseParamValue("ReportFilename");
 				if (!string.IsNullOrEmpty(ReportFilename))
 				{
 					IOProvider = new ReportIOProvider(ReportFilename);
@@ -50,12 +31,14 @@ namespace Turnkey
 				}
 			}
 
-			Turnkey.Execute(IOProvider, Build);
+			return Turnkey.Execute(IOProvider, this);
 		}
 
 
 		public static AutomationTool.ExitCode Execute(IOProvider IOProvider, BuildCommand CommandUtilHelper)
 		{
+			SetupVisuals();
+
 			// cache some settings for other classes
 			TurnkeyUtils.SetVariable("EngineDir", EngineDirectory.FullName);
 			TurnkeyUtils.SetVariable("Project", CommandUtilHelper.ParseParamValue("Project="));
@@ -103,5 +86,27 @@ namespace Turnkey
 
 			return TurnkeyUtils.ExitCode;
 		}
+
+		#region Visuals
+
+		[System.Runtime.InteropServices.DllImport("user32.dll")]
+		private static extern bool SetProcessDPIAware();
+
+		private static void SetupVisuals()
+		{
+			// make the form look good on modern displays!
+			if (!UnrealBuildTool.Utils.IsRunningOnMono && Environment.OSVersion.Version.Major >= 6)
+			{
+				SetProcessDPIAware();
+			}
+
+#if WINDOWS
+			System.Windows.Forms.Application.EnableVisualStyles();
+			System.Windows.Forms.Application.SetCompatibleTextRenderingDefault(false);
+#endif
+		}
+
+		#endregion
+
 	}
 }
diff --git a/Engine/Source/Programs/AutomationTool/Turnkey/TurnkeySettings.cs b/Engine/Source/Programs/AutomationTool/Turnkey/TurnkeySettings.cs
--- a/Engine/Source/Programs/AutomationTool/Turnkey/TurnkeySettings.cs
+++ b/Engine/Source/Programs/AutomationTool/Turnkey/TurnkeySettings.cs
@@ -37,10 +37,13 @@ namespace Turnkey
 		public static UserSetting[] AllUserSettings =
 		{
 			new UserSetting("User_QuickSwitchSdkLocation", null, "Location for downloaded Sdks to be stored (currently just for Google Drive, perforce will sync per clientspec)"),
-			new UserSetting("User_LastPerforceClient", null, "Perforce depot clientspec to look at first for finding matching clientspecs. Will be set automatically when one is found, but this can be used to force a particular client"),
 			new UserSetting("User_AppleDevCenterUsername", null, "Sets the username to use when logging in to DevCenter for iOS/tvOS/macOS. (Supercedes the DevCenterUsername .ini setting)"),
 			new UserSetting("User_AppleDevCenterTeamID", null, "Sets the teamid to use when logging in to DevCenter for iOS/tvOS/macOS. (Supercedes the IOSTeamID .ini setting)"),
 			new UserSetting("User_IOSProvisioningProfile", null, "Sets the provisioning profile to use when setting up iOS signing. (Supercedes the MobileProvision .ini setting)"),
+			new UserSetting("User_LastPerforceClient", null, "Perforce depot clientspec to look at first for finding matching clientspecs. Will be set automatically when one is found, but this can be used to force a particular client"),
+			new UserSetting("User_ManualProjects", null, "List of manually entered projects, separated by semicolons. Set automatically, but can be edited to clean it up."),
+			new UserSetting("User_LastSelectedProject", null, "The most recent project selected. Set automatically."),
+			new UserSetting("User_LastSelectedPlatform", null, "The most recent platform selected. Set automatically."),
 		};
 
 		public static UserSetting[] AllStudioSettings =
@@ -116,6 +119,11 @@ namespace Turnkey
 			return Value;
 		}
 
+		public static string GetUserSettingIfSet(string VariableName, string DefaultIfNotSet)
+		{
+			return HasSetUserSetting(VariableName) ? GetUserSetting(VariableName) : DefaultIfNotSet;
+		}
+
 		public static string SetUserSetting(string VariableName, string Value)
 		{
 			string PreviousValue = null;
diff --git a/Engine/Source/Programs/AutomationTool/Turnkey/TurnkeyUtils.cs b/Engine/Source/Programs/AutomationTool/Turnkey/TurnkeyUtils.cs
--- a/Engine/Source/Programs/AutomationTool/Turnkey/TurnkeyUtils.cs
+++ b/Engine/Source/Programs/AutomationTool/Turnkey/TurnkeyUtils.cs
@@ -118,6 +118,9 @@ namespace Turnkey
 			// HACK UNTIL WIN32 IS GONE
 			PossiblePlatforms = PossiblePlatforms.Where(x => x != UnrealTargetPlatform.Win32 && x != UnrealTargetPlatform.XboxOne).ToList();
 
+			// sort by name
+			PossiblePlatforms.Sort((x, y) => string.Compare(x.ToString(), y.ToString()));
+
 			List<UnrealTargetPlatform> Platforms = new List<UnrealTargetPlatform>();
 			// prompt user for a platform
 			if (PlatformString == null)
@@ -128,9 +131,14 @@ namespace Turnkey
 					return null;
 				}
 
+				// default to last platform chosen
+				string LastPlatform = TurnkeySettings.GetUserSettingIfSet("User_LastSelectedPlatform", "");
+
 				List<string> PlatformOptions = PossiblePlatforms.ConvertAll(x => x.ToString());
 				PlatformOptions.Add("All of the Above");
-				int PlatformChoice = TurnkeyUtils.ReadInputInt("Choose a platform:", PlatformOptions, true);
+
+				int Default = PlatformOptions.FindIndex(x => x.Equals(LastPlatform, StringComparison.OrdinalIgnoreCase));
+				int PlatformChoice = TurnkeyUtils.ReadInputInt("Choose a platform:", PlatformOptions, true, Default == -1 ? -1 : Default + 1);
 
 				if (PlatformChoice == 0)
 				{
@@ -144,6 +152,8 @@ namespace Turnkey
 				else
 				{
 					Platforms.Add(PossiblePlatforms[PlatformChoice - 1]);
+					// remember for next time
+					TurnkeySettings.SetUserSetting("User_LastSelectedPlatform", PossiblePlatforms[PlatformChoice - 1].ToString());
 				}
 			}
 			else if (PlatformString.ToLower() == "all")
@@ -174,6 +184,198 @@ namespace Turnkey
 			return Platforms.OrderBy(x => x.ToString()).ToList();
 		}
 
+		private enum ProjectMode
+		{
+			Normal,
+			Templates,
+			Sandbox,
+			Misc,
+			Samples,
+			Collaboration,
+			Personal,
+			Recent,
+		}
+
+		public static FileReference GetProjectFromCommandLineOrUser(string[] CommandOptions)
+		{
+			ProjectMode CurrentMode = ProjectMode.Normal;
+
+			string Project = TurnkeyUtils.GetVariableValue("Project");
+			string LastProject = TurnkeySettings.GetUserSettingIfSet("User_LastSelectedProject", "");
+
+			// create and initialize a dictionary of type to project list
+			Dictionary<ProjectMode, List<FileReference>> ProjectsByType = new Dictionary<ProjectMode, List<FileReference>>();
+			foreach (ProjectMode Mode in Enum.GetValues(typeof(ProjectMode)))
+			{
+				ProjectsByType.Add(Mode, new List<FileReference>());
+			}
+
+			// look up old manual projects
+			string ManualProjectString = TurnkeySettings.GetUserSettingIfSet("User_ManualProjects", "");
+			List<string> ManualProjects = ManualProjectString.Split(";", StringSplitOptions.RemoveEmptyEntries).ToList();
+			ProjectsByType[ProjectMode.Recent].AddRange(ManualProjects.Select(x => new FileReference(x)).ToList());
+			// start on Recent page if last one was a manual project
+			if (ProjectsByType[ProjectMode.Recent].Any(x => x.FullName.Equals(LastProject, StringComparison.OrdinalIgnoreCase)))
+			{
+				CurrentMode = ProjectMode.Recent;
+			}
+
+			// sort the discovered projects
+			foreach (FileReference NativeProject in NativeProjects.EnumerateProjectFiles())
+			{
+				DirectoryReference ProjectDir = NativeProject.Directory;
+				ProjectMode Mode = ProjectMode.Normal;
+
+				if (ProjectDir.ContainsName("StarterContent", 0) || ProjectDir.ContainsName("ContentExamples", 0))
+				{
+					Mode = ProjectMode.Misc;
+				}
+				else if (ProjectDir.ContainsName("Sandbox", 0))
+				{
+					Mode = ProjectMode.Sandbox;
+				}
+				else if (ProjectDir.ContainsName("Templates", 0))
+				{
+					Mode = ProjectMode.Templates;
+				}
+				else if (ProjectDir.ContainsName("Samples", 0))
+				{
+					Mode = ProjectMode.Samples;
+				}
+				else if (ProjectDir.ContainsName("Collaboration", 0))
+				{
+					Mode = ProjectMode.Collaboration;
+				}
+
+				// if the last project was this project, then use this as the starting type
+				if (LastProject != null && NativeProject.GetFileNameWithoutAnyExtensions().Equals(LastProject, StringComparison.OrdinalIgnoreCase))
+				{
+					CurrentMode = Mode;
+				}
+				ProjectsByType[Mode].Add(NativeProject);
+			}
+
+			// now look for the My Documents location
+			DirectoryReference PersonalProjects = new DirectoryReference(Path.Combine(Environment.GetFolderPath(Environment.SpecialFolder.Personal), "Unreal Projects"));
+			foreach (DirectoryReference PersonalProjectDir in DirectoryReference.EnumerateDirectories(PersonalProjects))
+			{
+				FileReference PersonalProjectFile = DirectoryReference.EnumerateFiles(PersonalProjectDir, "*.uproject").FirstOrDefault();
+				if (PersonalProjectFile != null)
+				{
+					ProjectsByType[ProjectMode.Personal].Add(PersonalProjectFile);
+				}
+			}
+
+
+			while (string.IsNullOrEmpty(Project))
+			{
+				// get list of project names from the dictionary
+				List<string> ProjectNames = new List<string>();
+
+				// add options to switch mode
+				Dictionary<int, ProjectMode> IndexToMode = new Dictionary<int, ProjectMode>();
+				foreach (ProjectMode Mode in Enum.GetValues(typeof(ProjectMode)))
+				{
+					if (Mode != CurrentMode)
+					{
+						IndexToMode.Add(ProjectNames.Count, Mode);
+						ProjectNames.Add(string.Format("[Show {0} Projects]", Mode.ToString()));
+					}
+				}
+
+				// now show the projects in this mode
+				ProjectNames.AddRange(ProjectsByType[CurrentMode].Select(x => (CurrentMode == ProjectMode.Recent) ? x.FullName : x.GetFileNameWithoutAnyExtensions()));
+
+				// and finally a manual entry option
+#if WINDOWS
+				ProjectNames.Add("[Browse...]");
+#else
+				ProjectNames.Add("[Enter Path To .uproject...]");
+#endif
+
+				int Default = LastProject == null ? -1 : ProjectNames.FindIndex(x => x.Equals(LastProject, StringComparison.OrdinalIgnoreCase));
+				int Choice = TurnkeyUtils.ReadInputInt(string.Format("Choose a {0} project to execute, or select a set of projects to list. (You can skip this by specifying -project=XXX on the commandline)", CurrentMode), ProjectNames, true, Default == -1 ? -1 : Default + 1);
+				if (Choice == 0)
+				{
+					return null;
+				}
+				// skip cancel
+				Choice = Choice - 1;
+
+				// look for manual entry option
+				if (Choice == ProjectNames.Count - 1)
+				{
+#if WINDOWS
+					System.Windows.Forms.OpenFileDialog Dialog = new System.Windows.Forms.OpenFileDialog();
+					Dialog.Filter = "Project Files (*.uproject)|*.uproject";
+					Dialog.AutoUpgradeEnabled = true;
+					Dialog.CheckFileExists = true;
+//						Dialog.InitialDirectory = TurnkeySettings.GetUserSettingIfSet("User_Last")
+
+					System.Windows.Forms.DialogResult Result = System.Windows.Forms.DialogResult.None;
+					System.Threading.Thread t = new System.Threading.Thread(x =>
+					{
+						Result = Dialog.ShowDialog();
+					});
+
+					t.SetApartmentState(System.Threading.ApartmentState.STA);
+					t.Start();
+					t.Join();
+
+					if (Result != System.Windows.Forms.DialogResult.OK)
+					{
+						continue;
+					}
+					Project = Dialog.FileName;
+#else
+					while (true)
+					{
+						string ProjectChoice = TurnkeyUtils.ReadInput("Enter path to .uproject file:");
+						if (!File.Exists(ProjectChoice))
+						{
+							string Response = TurnkeyUtils.ReadInput(string.Format("'{0}' doesn't exist. Would you like to enter another path? [y/N]", ProjectChoice), "N");
+							if (Response.ToLower() != "y")
+							{
+								break;
+							}
+						}
+						else
+						{
+							Project = ProjectChoice;
+							break;
+						}
+					}
+#endif
+
+					// if we got a valid vhoice, remember it
+					if (!string.IsNullOrEmpty(Project))
+					{
+						// remember this for next run
+						// case insensitive Contains
+						if (!ManualProjects.Any(x => x.Equals(Project, StringComparison.OrdinalIgnoreCase)))
+						{
+							ManualProjects.Add(Project);
+						}
+						TurnkeySettings.SetUserSetting("User_ManualProjects", string.Join(";", ManualProjects));
+					}
+				}
+				else
+				{
+					// if the choice was a mode switch, get the mode
+					if (IndexToMode.TryGetValue(Choice, out CurrentMode))
+					{
+						continue;
+					}
+					Project = ProjectNames[Choice];
+				}
+			}
+
+			// remember project for next run
+			TurnkeySettings.SetUserSetting("User_LastSelectedProject", Project);
+
+			return ProjectUtils.FindProjectFileFromName(Project);
+		}
+
 		public static DeviceInfo GetDeviceFromCommandLineOrUser(string[] CommandOptions, UnrealTargetPlatform Platform)
 		{
 			string DeviceName = TurnkeyUtils.ParseParamValue("Device", null, CommandOptions);
@@ -212,9 +414,9 @@ namespace Turnkey
 			DeviceInfo InstallDevice = Array.Find(AutomationPlatform.GetDevices(), x => string.Compare(x.Name, DeviceName, true) == 0);
 			return InstallDevice;
 		}
-		#endregion
+#endregion
 
-		#region Env vars
+#region Env vars
 
 		private static IDictionary[] SavedEnvVars = new System.Collections.IDictionary[2];
 		private static Dictionary<string, string> EnvVarsToSaveToBatchFile = new Dictionary<string, string>();
@@ -275,9 +477,9 @@ namespace Turnkey
 			}
 		}
 
-		#endregion
+#endregion
 
-		#region Regex Matching
+#region Regex Matching
 
 		static bool TryConvertToUint64(string InValue, out UInt64 OutValue)
 		{ 
@@ -349,9 +551,9 @@ namespace Turnkey
 			return string.Compare(Value, AllowedValues, true) == 0;
 		}
 
-		#endregion
+#endregion
 
-		#region IO
+#region IO
 		static private IOProvider IOProvider;
 
 		public static void Log(string Message)
@@ -396,9 +598,9 @@ namespace Turnkey
 		{
 			return IOProvider.ReadInputInt(Prompt, Options, bIsCancellable, DefaultValue, bAppendNewLine: true);
 		}
-		#endregion
+#endregion
 
-		#region Temp files [move to LocalCache]
+#region Temp files [move to LocalCache]
 		static List<string> PathsToCleanup = new List<string>();
 		public static void AddPathToCleanup(string Path)
 		{
@@ -424,6 +626,6 @@ namespace Turnkey
 
 			PathsToCleanup.Clear();
 		}
-		#endregion
+#endregion
 	}
 }
