- Hopeful fix for hard-to-repro crash

diff --git a/Engine/Source/Editor/TurnkeySupport/Private/TurnkeySupportModule.cpp b/Engine/Source/Editor/TurnkeySupport/Private/TurnkeySupportModule.cpp
--- a/Engine/Source/Editor/TurnkeySupport/Private/TurnkeySupportModule.cpp
+++ b/Engine/Source/Editor/TurnkeySupport/Private/TurnkeySupportModule.cpp
@@ -1128,7 +1128,7 @@ void FTurnkeySupportModule::MakeQuickLaunchItems(class UToolMenu* Menu, FOnQuick
 						NAME_None,
 						MakeSdkStatusAttribute(PlatformName, DeviceProxies[0]),
 						Tooltip,
-						FNewMenuDelegate::CreateLambda([TargetDeviceServicesModule, PlatformName, &ExternalOnClickDelegate](FMenuBuilder& SubMenuBuilder)
+						FNewMenuDelegate::CreateLambda([TargetDeviceServicesModule, PlatformName, ExternalOnClickDelegate](FMenuBuilder& SubMenuBuilder)
 							{
 								// re-get the proxies, just in case they changed
 								TArray<TSharedPtr<ITargetDeviceProxy>> DeviceProxies;
