- Fixed chosen configration for packaging getting lost from Platforms menu

diff --git a/Engine/Source/Editor/TurnkeySupport/Private/TurnkeySupportModule.cpp b/Engine/Source/Editor/TurnkeySupport/Private/TurnkeySupportModule.cpp
--- a/Engine/Source/Editor/TurnkeySupport/Private/TurnkeySupportModule.cpp
+++ b/Engine/Source/Editor/TurnkeySupport/Private/TurnkeySupportModule.cpp
@@ -313,20 +313,20 @@ public:
 			ContentPrepIcon = FEditorStyle::GetBrush(TEXT("MainFrame.PackageProject"));
 
 			// let the user pick a target directory
-			if (PackagingSettings->StagingDirectory.Path.IsEmpty())
+			if (AllPlatformPackagingSettings->StagingDirectory.Path.IsEmpty())
 			{
-				PackagingSettings->StagingDirectory.Path = FPaths::ProjectDir();
+				AllPlatformPackagingSettings->StagingDirectory.Path = FPaths::ProjectDir();
 			}
 
 			FString OutFolderName;
 
-			if (!FDesktopPlatformModule::Get()->OpenDirectoryDialog(FSlateApplication::Get().FindBestParentWindowHandleForDialogs(nullptr), LOCTEXT("PackageDirectoryDialogTitle", "Package project...").ToString(), PackagingSettings->StagingDirectory.Path, OutFolderName))
+			if (!FDesktopPlatformModule::Get()->OpenDirectoryDialog(FSlateApplication::Get().FindBestParentWindowHandleForDialogs(nullptr), LOCTEXT("PackageDirectoryDialogTitle", "Package project...").ToString(), AllPlatformPackagingSettings->StagingDirectory.Path, OutFolderName))
 			{
 				return;
 			}
 
-			PackagingSettings->StagingDirectory.Path = OutFolderName;
-			PackagingSettings->SaveConfig();
+			AllPlatformPackagingSettings->StagingDirectory.Path = OutFolderName;
+			AllPlatformPackagingSettings->SaveConfig();
 
 
 			BuildCookRunParams += TEXT(" -stage -archive -package");
@@ -1608,7 +1608,7 @@ void FTurnkeySupportModule::UpdateSdkInfoForDevices(TArray<FString> PlatformDevi
 
 	FString Commandline = BaseCommandline + FString(TEXT(" -Device=")) + FString::JoinBy(PlatformDeviceIds, TEXT("+"), [](FString Id) { return ConvertToUATDeviceId(Id); });
 
-	UE_LOG(LogTurnkeySupport, Log, TEXT("Running Turnkey SDK detection: '%s %s'"), *Command, *Commandline);
+	UE_LOG(LogTurnkeySupport, Log, TEXT("Running Turnkey device detection: '%s %s'"), *Command, *Commandline);
 
 	{
 		FScopeLock Lock(&GTurnkeySection);
