- Fixed bug I introduced with dropping the -archivedirectory option

diff --git a/Engine/Source/Editor/TurnkeySupport/Private/TurnkeySupportModule.cpp b/Engine/Source/Editor/TurnkeySupport/Private/TurnkeySupportModule.cpp
--- a/Engine/Source/Editor/TurnkeySupport/Private/TurnkeySupportModule.cpp
+++ b/Engine/Source/Editor/TurnkeySupport/Private/TurnkeySupportModule.cpp
@@ -374,7 +374,7 @@ public:
 				BuildCookRunParams += TEXT(" -applocaldirectory=\"$(EngineDir)/Binaries/ThirdParty/AppLocalDependencies\"");
 			}
 
-			BuildCookRunParams += FString::Printf(TEXT(" -archivedirectory=\"%s\""), *PackagingSettings->StagingDirectory.Path);
+			BuildCookRunParams += FString::Printf(TEXT(" -archivedirectory=\"%s\""), *AllPlatformPackagingSettings->StagingDirectory.Path);
 
 			if (PackagingSettings->ForDistribution)
 			{
