Turnkey: Allow the project path to be empty

#rnx

diff --git a/Engine/Source/Editor/TurnkeySupport/Private/TurnkeySupportModule.cpp b/Engine/Source/Editor/TurnkeySupport/Private/TurnkeySupportModule.cpp
--- a/Engine/Source/Editor/TurnkeySupport/Private/TurnkeySupportModule.cpp
+++ b/Engine/Source/Editor/TurnkeySupport/Private/TurnkeySupportModule.cpp
@@ -58,6 +58,27 @@ namespace
 	FCriticalSection GTurnkeySection;
 }
 
+static FString GetProjectPathForTurnkey()
+{
+	if (FPaths::IsProjectFilePathSet())
+	{
+		return FPaths::ConvertRelativePathToFull(FPaths::GetProjectFilePath());
+	}
+	if (FApp::HasProjectName())
+	{
+		FString ProjectPath = FPaths::ProjectDir() / FApp::GetProjectName() + TEXT(".uproject");
+		if (FPaths::FileExists(ProjectPath))
+		{
+			return ProjectPath;
+		}
+		ProjectPath = FPaths::RootDir() / FApp::GetProjectName() / FApp::GetProjectName() + TEXT(".uproject");
+		if (FPaths::FileExists(ProjectPath))
+		{
+			return ProjectPath;
+		}
+	}
+	return FString();
+}
 
 enum class EPrepareContentMode : uint8
 {
@@ -235,7 +256,7 @@ public:
 		check(PlatformInfo != nullptr);
 
 		const FString UBTPlatformString = PlatformInfo->DataDrivenPlatformInfo->UBTPlatformString;
-		const FString ProjectPath = FPaths::IsProjectFilePathSet() ? FPaths::ConvertRelativePathToFull(FPaths::GetProjectFilePath()) : FPaths::RootDir() / FApp::GetProjectName() / FApp::GetProjectName() + TEXT(".uproject");
+		const FString ProjectPath = GetProjectPathForTurnkey();
 
 
 		// check that we can proceed
@@ -266,6 +287,7 @@ public:
 
 
 		// set locations to engine and project
+		if (!ProjectPath.IsEmpty())
 		{
 			BuildCookRunParams += FString::Printf(TEXT(" -project=\"%s\""), *ProjectPath);
 		}
@@ -440,13 +462,18 @@ public:
 		}
 
 
-		FString TurnkeyParams = FString::Printf(TEXT("-command=VerifySdk -platform=%s -UpdateIfNeeded -EditorIO -project=\"%s\""), *UBTPlatformString, *ProjectPath);
+		FString TurnkeyParams = FString::Printf(TEXT("-command=VerifySdk -platform=%s -UpdateIfNeeded -EditorIO"), *UBTPlatformString);
+		if (!ProjectPath.IsEmpty())
+		{
+			TurnkeyParams.Appendf(TEXT(" -project=\"%s\""), *ProjectPath);
+		}
 
-		FString CommandLine = FString::Printf(TEXT("-ScriptsForProject=\"%s\" Turnkey %s BuildCookRun %s"),
-			*ProjectPath,
-			*TurnkeyParams,
-			*BuildCookRunParams
-		);
+		FString CommandLine;
+		if (!ProjectPath.IsEmpty())
+		{
+			CommandLine.Appendf(TEXT("-ScriptsForProject=\"%s\" "), *ProjectPath);
+		}
+		CommandLine.Appendf(TEXT("Turnkey %s BuildCookRun %s"), *TurnkeyParams, *BuildCookRunParams);
 
 		FTurnkeyEditorSupport::RunUAT(CommandLine, PlatformInfo->DisplayName, ContentPrepDescription, ContentPrepTaskName, ContentPrepIcon);
 	}
@@ -464,10 +491,19 @@ public:
 	static void ExecuteCustomBuild(FName IniPlatformName, FProjectBuildSettings Build)
 	{
 		const PlatformInfo::FTargetPlatformInfo* PlatformInfo = PlatformInfo::FindPlatformInfo(GetDefault<UProjectPackagingSettings>()->GetTargetPlatformForPlatform(IniPlatformName));
-		const FString ProjectPath = FPaths::IsProjectFilePathSet() ? FPaths::ConvertRelativePathToFull(FPaths::GetProjectFilePath()) : FPaths::RootDir() / FApp::GetProjectName() / FApp::GetProjectName() + TEXT(".uproject");
+		const FString ProjectPath = GetProjectPathForTurnkey();
 
-		FString CommandLine = FString::Printf(TEXT("-ScriptsForProject=\"%s\" Turnkey -command=ExecuteBuild -build=\"%s\" -platform=%s -project=\"%s\""),
-			*ProjectPath, *Build.Name, *IniPlatformName.ToString(), *ProjectPath);
+		FString CommandLine;
+		if (!ProjectPath.IsEmpty())
+		{
+			CommandLine.Appendf(TEXT("-ScriptsForProject=\"%s\" "), *ProjectPath);
+		}
+		CommandLine.Appendf(TEXT("Turnkey -command=ExecuteBuild -build=\"%s\" -platform=%s"),
+			*Build.Name, *IniPlatformName.ToString());
+		if (!ProjectPath.IsEmpty())
+		{
+			CommandLine.Appendf(TEXT(" -project=\"%s\""), *ProjectPath);
+		}
 
 		FTurnkeyEditorSupport::RunUAT(CommandLine, PlatformInfo->DisplayName, LOCTEXT("Turnkey_CustomTaskName", "Executing Custom Build"), LOCTEXT("Turnkey_CustomTaskName", "Custom"), FEditorStyle::GetBrush(TEXT("MainFrame.PackageProject")));
 	}
@@ -619,8 +655,13 @@ static void TurnkeyInstallSdk(FString PlatformName, bool bPreferFull, bool bForc
 		OptionalOptions += FString::Printf(TEXT(" -Device=%s"), *DeviceId);
 	}
 
-	const FString ProjectPath = FPaths::IsProjectFilePathSet() ? FPaths::ConvertRelativePathToFull(FPaths::GetProjectFilePath()) : FPaths::RootDir() / FApp::GetProjectName() / FApp::GetProjectName() + TEXT(".uproject");
-	FString CommandLine = FString::Printf(TEXT("-ScriptsForProject=\"%s\" Turnkey -command=VerifySdk -UpdateIfNeeded -platform=%s %s -EditorIO -noturnkeyvariables -utf8output -WaitForUATMutex"), *ProjectPath, *PlatformName, *OptionalOptions);
+	const FString ProjectPath = GetProjectPathForTurnkey();
+	FString CommandLine;
+	if (!ProjectPath.IsEmpty())
+	{
+		CommandLine.Appendf(TEXT("-ScriptsForProject=\"%s\" "), *ProjectPath);
+	}
+	CommandLine.Appendf(TEXT("Turnkey -command=VerifySdk -UpdateIfNeeded -platform=%s %s -EditorIO -noturnkeyvariables -utf8output -WaitForUATMutex"), *PlatformName, *OptionalOptions);
 
 	FText TaskName = LOCTEXT("InstallingSdk", "Installing Sdk");
 	FTurnkeyEditorSupport::RunUAT(CommandLine, FText::FromString(PlatformName), TaskName, TaskName, FEditorStyle::GetBrush(TEXT("MainFrame.PackageProject")),
@@ -1385,8 +1426,13 @@ static void PrepForTurnkeyReport(FString& Command, FString& BaseCommandline, FSt
 	Command = TEXT("{EngineDir}Build/BatchFiles/RunuAT");
 	//	Command = TEXT("{EngineDir}/Binaries/DotNET/AutomationTool.exe");
 
-	const FString ProjectPath = FPaths::IsProjectFilePathSet() ? FPaths::ConvertRelativePathToFull(FPaths::GetProjectFilePath()) : FPaths::RootDir() / FApp::GetProjectName() / FApp::GetProjectName() + TEXT(".uproject");
-	BaseCommandline = FString::Printf(TEXT("-ScriptsForProject=\"%s\" Turnkey -utf8output -WaitForUATMutex -command=VerifySdk -ReportFilename=\"%s\" -log=\"%s\""), *ProjectPath, *ReportFilename, *LogFilename);
+	BaseCommandline.Empty();
+	const FString ProjectPath = GetProjectPathForTurnkey();
+	if (!ProjectPath.IsEmpty())
+	{
+		BaseCommandline.Appendf(TEXT("-ScriptsForProject=\"%s\" "), *ProjectPath);
+	}
+	BaseCommandline.Appendf(TEXT("Turnkey -utf8output -WaitForUATMutex -command=VerifySdk -ReportFilename=\"%s\" -log=\"%s\""), *ReportFilename, *LogFilename);
 
 	// convert into appropriate calls for the current platform
 	FPlatformProcess::ModifyCreateProcParams(Command, BaseCommandline, FGenericPlatformProcess::ECreateProcHelperFlags::AppendPlatformScriptExtension | FGenericPlatformProcess::ECreateProcHelperFlags::RunThroughShell);
