- Added FSerializedUATProcess - may move to a header at some point
- Removed ModifyCreateProcParams., it wasn't useful enough
- Split up ExtraProjectBuilds into EngineCustomBuilds and ProjectCustomBuilds, so that engine/Epic specified custom builds don't dirty up the project packaging settings, and the ini system won't cause future additional builds to be lost
- Fixed up a couple of tooltips to indicate compression settings are used for Pak and IoStore

diff --git a/Engine/Config/BaseGame.ini b/Engine/Config/BaseGame.ini
--- a/Engine/Config/BaseGame.ini
+++ b/Engine/Config/BaseGame.ini
@@ -80,7 +80,8 @@ MinWindowWidth=16
 MinWindowHeight=16
 
 [/Script/UnrealEd.ProjectPackagingSettings]
-@ExtraProjectBuilds=Name
+@EngineCustomBuilds=Name
+@ProjectCustomBuilds=Name
 BuildConfiguration=PPBC_Development
 FullRebuild=False
 UsePakFile=True
diff --git a/Engine/Source/Developer/DeveloperToolSettings/Classes/Settings/ProjectPackagingSettings.h b/Engine/Source/Developer/DeveloperToolSettings/Classes/Settings/ProjectPackagingSettings.h
--- a/Engine/Source/Developer/DeveloperToolSettings/Classes/Settings/ProjectPackagingSettings.h
+++ b/Engine/Source/Developer/DeveloperToolSettings/Classes/Settings/ProjectPackagingSettings.h
@@ -271,16 +271,16 @@ public:
 	FDirectoryPath HttpChunkInstallDataDirectory;
 
 	/**
-	 * A comma separated list of formats to use for .pak file compression. If more than one is specified, the list is in order of priority, with fallbacks to other formats
+	 * A comma separated list of formats to use for .pak file and IoStore compression. If more than one is specified, the list is in order of priority, with fallbacks to other formats
 	 * in case of errors or unavailability of the format (plugin not enabled, etc).
 	 */
-	UPROPERTY(config, EditAnywhere, Category = Packaging, meta = (DisplayName = "Pak File Compression Format(s)"))
+	UPROPERTY(config, EditAnywhere, Category = Packaging, meta = (DisplayName = "Pak File / IoStore Compression Format(s)"))
 	FString PakFileCompressionFormats;
 
 	/**
-	 * A generic setting for allowing a project to control compression settings during .pak file compression. For instance, if using the Oodle plugin, you could use -oodlemethod=Selkie -oodlelevel=Optimal1
+	 * A generic setting for allowing a project to control compression settings during .pak file and IOStore compression. For instance, if using the Oodle plugin, you could use -oodlemethod=Selkie -oodlelevel=Optimal1
 	 */
-	UPROPERTY(config, EditAnywhere, Category = Packaging, meta = (DisplayName = "Pak File Compression Commandline Options"))
+	UPROPERTY(config, EditAnywhere, Category = Packaging, meta = (DisplayName = "Pak File / IoStore Compression Commandline Options"))
 	FString PakFileAdditionalCompressionOptions;
 
 	/** 
@@ -502,7 +502,13 @@ public:
 	 * A list of custom builds that will show up in the Platforms menu to allow customized builds that make sense for your project. Will show up near Package Project in the Platforms menu.
 	 */
 	UPROPERTY(config, EditAnywhere, Category=Packaging, meta=(DisplayName = "Additional builds for this project."))
-	TArray<FProjectBuildSettings> ExtraProjectBuilds;
+	TArray<FProjectBuildSettings> ProjectCustomBuilds;
+
+	/**
+	 * A list of custom builds, specified in engine ini files, and not editable in editor, that will show up in the Platforms menu to allow customized builds for all projects
+	 */
+	UPROPERTY(config)
+	TArray<FProjectBuildSettings> EngineCustomBuilds;
 
 
 	/**
diff --git a/Engine/Source/Editor/TurnkeySupport/Private/TurnkeySupportModule.cpp b/Engine/Source/Editor/TurnkeySupport/Private/TurnkeySupportModule.cpp
--- a/Engine/Source/Editor/TurnkeySupport/Private/TurnkeySupportModule.cpp
+++ b/Engine/Source/Editor/TurnkeySupport/Private/TurnkeySupportModule.cpp
@@ -18,6 +18,7 @@
 #include "Async/Async.h"
 #include "Misc/FileHelper.h"
 #include "Misc/ConfigCacheIni.h"
+#include "Delegates/Delegate.h"
 #include "Interfaces/TargetDeviceId.h"
 #include "Interfaces/ITargetPlatform.h"
 #include "Interfaces/ITargetPlatformManagerModule.h"
@@ -816,13 +817,31 @@ static void MakeTurnkeyPlatformMenu(FMenuBuilder& MenuBuilder, FName IniPlatform
 		FString PlatformString = IniPlatformName.ToString();
 		UProjectPackagingSettings* PackagingSettings = FTurnkeySupportCallbacks::GetPackagingSettingsForPlatform(IniPlatformName);
 
-		for (FProjectBuildSettings Build : PackagingSettings->ExtraProjectBuilds)
+		for (FProjectBuildSettings Build : PackagingSettings->EngineCustomBuilds)
 		{
 			if (Build.SpecificPlatforms.Num() == 0 || Build.SpecificPlatforms.Contains(PlatformString))
 			{
 				MenuBuilder.AddMenuEntry(
 					FText::FromString(Build.Name),
-					LOCTEXT("TurnkeyTooltip_CustomBuild", "Execute a custom build (this comes from Packaging Settings)"),
+					// @todo turnkey: add the build string to the tooltip
+					LOCTEXT("TurnkeyTooltip_EngineCustomBuild", "Execute a custom build"),
+					FSlateIcon(),
+					FUIAction(
+						FExecuteAction::CreateStatic(&FTurnkeySupportCallbacks::ExecuteCustomBuild, IniPlatformName, Build),
+						FCanExecuteAction::CreateStatic(&FTurnkeySupportCallbacks::CanExecuteCustomBuild, IniPlatformName, Build)
+					)
+				);
+			}
+		}
+
+		for (FProjectBuildSettings Build : PackagingSettings->ProjectCustomBuilds)
+		{
+			if (Build.SpecificPlatforms.Num() == 0 || Build.SpecificPlatforms.Contains(PlatformString))
+			{
+				MenuBuilder.AddMenuEntry(
+					FText::FromString(Build.Name),
+					// @todo turnkey: add the build string to the tooltip
+					LOCTEXT("TurnkeyTooltip_ProjectCustomBuild", "Execute a custom build (this comes from Packaging Settings)"),
 					FSlateIcon(),
 					FUIAction(
 						FExecuteAction::CreateStatic(&FTurnkeySupportCallbacks::ExecuteCustomBuild, IniPlatformName, Build),
@@ -1423,19 +1442,25 @@ static void PrepForTurnkeyReport(FString& Command, FString& BaseCommandline, FSt
 	// make sure intermediate directory exists
 	IFileManager::Get().MakeDirectory(*FPaths::ProjectIntermediateDir());
 
-	Command = TEXT("{EngineDir}Build/BatchFiles/RunuAT");
-	//	Command = TEXT("{EngineDir}/Binaries/DotNET/AutomationTool.exe");
 
-	BaseCommandline.Empty();
+#if PLATFORM_WINDOWS
+	Command = TEXT("cmd.exe");
+	BaseCommandline = TEXT("/c RunUAT.bat");
+#elif PLATFORM_MAC
+	Command = TEXT("/bin/sh");
+	BaseCommandline = TEXT("RunUAT.sh");
+#elif PLATFORM_LINUX
+	Command = TEXT("/bin/bash");
+	BaseCommandline = TEXT("-c RunUAT.sh");
+#endif
+
 	const FString ProjectPath = GetProjectPathForTurnkey();
 	if (!ProjectPath.IsEmpty())
 	{
-		BaseCommandline.Appendf(TEXT("-ScriptsForProject=\"%s\" "), *ProjectPath);
+		BaseCommandline.Appendf(TEXT(" -ScriptsForProject=\"%s\" "), *ProjectPath);
 	}
-	BaseCommandline.Appendf(TEXT("Turnkey -utf8output -WaitForUATMutex -command=VerifySdk -ReportFilename=\"%s\" -log=\"%s\""), *ReportFilename, *LogFilename);
 
-	// convert into appropriate calls for the current platform
-	FPlatformProcess::ModifyCreateProcParams(Command, BaseCommandline, FGenericPlatformProcess::ECreateProcHelperFlags::AppendPlatformScriptExtension | FGenericPlatformProcess::ECreateProcHelperFlags::RunThroughShell);
+	BaseCommandline = BaseCommandline.Appendf(TEXT("-ScriptsForProject=\"%s\" Turnkey -utf8output -WaitForUATMutex -command=VerifySdk -ReportFilename=\"%s\" -log=\"%s\""), *ProjectPath, *ReportFilename, *LogFilename);
 }
 
 bool GetSdkInfoFromTurnkey(FString Line, FName& PlatformName, FString& DeviceId, FTurnkeySdkInfo& SdkInfo)
@@ -1494,24 +1519,86 @@ bool GetSdkInfoFromTurnkey(FString Line, FName& PlatformName, FString& DeviceId,
 	return true;
 }
 
-// FDataDrivenPlatformInfo& DeviceIdToInfo(FString DeviceId, FString* OutDeviceName)
-// {
-// 	TArray<FString> PlatformAndDevice;
-// 	DeviceId.ParseIntoArray(PlatformAndDevice, TEXT("@"), true);
-// 
-// 	if (OutDeviceName)
-// 	{
-// 		*OutDeviceName = PlatformAndDevice[1];
-// 	}
-// 
-// 	FString DDPIPlatformName = ConvertToDDPIPlatform(PlatformAndDevice[0]);
-// 
-// 	checkf(DataDrivenPlatforms.Contains(*DDPIPlatformName), TEXT("DataDrivenPlatforms map did not contain the DDPI Platform %s"), *DDPIPlatformName);
-// 
-// 	// have to convert back to Windows from Win64
-// 	return DataDrivenPlatforms[*DDPIPlatformName];
-// 
-// }
+DECLARE_DELEGATE_ThreeParams(FOnSerializedProcessComplete, int32 /*ExitCode*/, bool /*bLaunchFailed*/, bool /*bWasCanceled*/);
+
+class FSerializedUATProcess
+{
+public:
+	FSerializedUATProcess(const FString& URL, const FString& Params, bool bDeleteOnCompletion, FOnSerializedProcessComplete CompletionDelegate)
+		: Process(nullptr)
+	{
+		Async(EAsyncExecution::Thread, [this, URL, Params, bDeleteOnCompletion, CompletionDelegate]()
+			{
+				// don't let this do anything if another process is running
+				Serializer.Lock();
+
+				FString FinalParams = Params;
+				if (bHasSucceededOnce)
+				{
+					FinalParams += TEXT(" -nocompile");
+				}
+
+				Process = new FMonitoredProcess(URL, FinalParams, true, false);
+				Process->OnCompleted().BindLambda([this, bDeleteOnCompletion, CompletionDelegate](int32 ExitCode)
+					{
+						if (ExitCode == 0 || ExitCode == 10)
+						{
+							bHasSucceededOnce = true;
+						}
+						CompletionDelegate.Execute(ExitCode, false, false);
+
+						// let another one in
+						Serializer.Unlock();
+
+						if (bDeleteOnCompletion)
+						{
+							delete this;
+						}
+					});
+				Process->OnCanceled().BindLambda([this, bDeleteOnCompletion, CompletionDelegate]()
+					{
+						CompletionDelegate.Execute(-1, false, true);
+
+						// let another one in
+						Serializer.Unlock();
+
+						if (bDeleteOnCompletion)
+						{
+							delete this;
+						}
+					});
+				
+				UE_LOG(LogTurnkeySupport, Log, TEXT("Running Serialized UAT: '%s %s'"), *URL, *FinalParams);
+
+				if (Process->Launch() == false)
+				{
+					CompletionDelegate.Execute(-1, true, false);
+
+					// let another one in
+					Serializer.Unlock();
+
+					if (bDeleteOnCompletion)
+					{
+						delete this;
+					}
+				}
+			});
+	}
+
+	~FSerializedUATProcess()
+	{
+		delete Process;
+	}
+
+private:
+
+	static bool bHasSucceededOnce;
+	static FCriticalSection Serializer;
+
+	FMonitoredProcess* Process;
+};
+bool FSerializedUATProcess::bHasSucceededOnce = false;
+FCriticalSection FSerializedUATProcess::Serializer;
 
 
 void FTurnkeySupportModule::UpdateSdkInfo()
@@ -1552,10 +1639,11 @@ void FTurnkeySupportModule::UpdateSdkInfo()
 		ClearDeviceStatus();
 	}
 
-	FMonitoredProcess* TurnkeyProcess = new FMonitoredProcess(Command, Commandline, true, false);
-	TurnkeyProcess->OnCompleted().BindLambda([this, ReportFilename, TurnkeyProcess](int32 ExitCode)
+	FSerializedUATProcess* TurnkeyProcess = new FSerializedUATProcess(Command, Commandline, true, FOnSerializedProcessComplete::CreateLambda([this, ReportFilename](int32 ExitCode, bool bLaunchFailed, bool bWasCanceled)
 	{
-		AsyncTask(ENamedThreads::GameThread, [this, ReportFilename, TurnkeyProcess, ExitCode]()
+		UE_LOG(LogTurnkeySupport, Log, TEXT("Completed SDK detection: Code = %d, bLaunchFailed = %d, bWasCanceled = %d"), ExitCode, bLaunchFailed, bWasCanceled);
+
+		AsyncTask(ENamedThreads::GameThread, [this, ReportFilename, ExitCode]()
 		{
 			FScopeLock Lock(&GTurnkeySection);
 
@@ -1636,13 +1724,10 @@ void FTurnkeySupportModule::UpdateSdkInfo()
 			}
 
 			// cleanup
-			delete TurnkeyProcess;
+//			delete TurnkeyProcess;
 			IFileManager::Get().Delete(*ReportFilename);
 		});
-	});
-
-	// run it
-	TurnkeyProcess->Launch();
+	}));
 }
 
 void FTurnkeySupportModule::UpdateSdkInfoForDevices(TArray<FString> PlatformDeviceIds)
@@ -1668,10 +1753,11 @@ void FTurnkeySupportModule::UpdateSdkInfoForDevices(TArray<FString> PlatformDevi
 		}
 	}
 
-	FMonitoredProcess* TurnkeyProcess = new FMonitoredProcess(Command, Commandline, true, false);
-	TurnkeyProcess->OnCompleted().BindLambda([this, ReportFilename, TurnkeyProcess, PlatformDeviceIds](int32 ExitCode)
+	FSerializedUATProcess* TurnkeyProcess = new FSerializedUATProcess(Command, Commandline, true, FOnSerializedProcessComplete::CreateLambda([this, ReportFilename, PlatformDeviceIds](int32 ExitCode, bool bLaunchFailed, bool bWasCanceled)
 	{
-		AsyncTask(ENamedThreads::GameThread, [this, ReportFilename, TurnkeyProcess, PlatformDeviceIds, ExitCode]()
+		UE_LOG(LogTurnkeySupport, Log, TEXT("Completed device detection: Code = %d, bLaunchFailed = %d, bWasCanceled = %d"), ExitCode, bLaunchFailed, bWasCanceled);
+
+		AsyncTask(ENamedThreads::GameThread, [this, ReportFilename, PlatformDeviceIds, ExitCode]()
 		{
 			FScopeLock Lock(&GTurnkeySection);
 
@@ -1724,13 +1810,10 @@ void FTurnkeySupportModule::UpdateSdkInfoForDevices(TArray<FString> PlatformDevi
 			}
 
 			// cleanup
-			delete TurnkeyProcess;
+//			delete TurnkeyProcess;
 			IFileManager::Get().Delete(*ReportFilename);
 		});
-	});
-
-	// run it
-	TurnkeyProcess->Launch();
+	}));
 }
 
 /**
diff --git a/Engine/Source/Programs/AutomationTool/Turnkey/Commands/ExecuteBuild.cs b/Engine/Source/Programs/AutomationTool/Turnkey/Commands/ExecuteBuild.cs
--- a/Engine/Source/Programs/AutomationTool/Turnkey/Commands/ExecuteBuild.cs
+++ b/Engine/Source/Programs/AutomationTool/Turnkey/Commands/ExecuteBuild.cs
@@ -31,7 +31,7 @@ namespace Turnkey.Commands
 
 			// attempt to match it!
 			Match Result = Regex.Match(Input, PrimaryRegex);
-			if (!Result.Success)
+			if (!Result.Success && AltRegex != null)
 			{
 				Result = Regex.Match(Input, AltRegex);
 			}
@@ -64,8 +64,20 @@ namespace Turnkey.Commands
 			{
 				ConfigHierarchy GameConfig = ConfigCache.ReadHierarchy(ConfigHierarchyType.Game, ProjectFile.Directory, Platform);
 
-				List<string> Builds;
-				GameConfig.GetArray("/Script/UnrealEd.ProjectPackagingSettings", "ExtraProjectBuilds", out Builds);
+				List<string> EngineBuilds;
+				List<string> ProjectBuilds;
+				GameConfig.GetArray("/Script/UnrealEd.ProjectPackagingSettings", "EngineCustomBuilds", out EngineBuilds);
+				GameConfig.GetArray("/Script/UnrealEd.ProjectPackagingSettings", "ProjectCustomBuilds", out ProjectBuilds);
+
+				List<string> Builds = new List<string>();
+				if (EngineBuilds != null)
+				{
+					Builds.AddRange(EngineBuilds);
+				}
+				if (ProjectBuilds != null)
+				{
+					Builds.AddRange(ProjectBuilds);
+				}
 
 				Dictionary<string, string> BuildCommands = new Dictionary<string, string>(StringComparer.InvariantCultureIgnoreCase);
 				if (Builds != null)
@@ -73,7 +85,7 @@ namespace Turnkey.Commands
 					foreach (string Build in Builds)
 					{
 						string Name = GetStructEntry(Build, "Name", false);
-						string SpecificPlatforms = GetStructEntry(Build, "Platforms", true);
+						string SpecificPlatforms = GetStructEntry(Build, "SpecificPlatforms", true);
 						string Params = GetStructEntry(Build, "BuildCookRunParams", false);
 
 						// make sure required entries are there
diff --git a/Engine/Source/Runtime/Core/Private/GenericPlatform/GenericPlatformProcess.cpp b/Engine/Source/Runtime/Core/Private/GenericPlatform/GenericPlatformProcess.cpp
--- a/Engine/Source/Runtime/Core/Private/GenericPlatform/GenericPlatformProcess.cpp
+++ b/Engine/Source/Runtime/Core/Private/GenericPlatform/GenericPlatformProcess.cpp
@@ -256,71 +256,6 @@ void FGenericPlatformProcess::TerminateProc( FProcHandle & ProcessHandle, bool K
 	UE_LOG(LogHAL, Fatal, TEXT("FGenericPlatformProcess::TerminateProc not implemented on this platform"));
 }
 
-static FString ReplaceProcHelperVariables(const FString& Str)
-{
-	return Str.Replace(TEXT("{Platform}"), ANSI_TO_TCHAR(FPlatformProperties::IniPlatformName())).
-		Replace(TEXT("{BaseDir}"), FPlatformProcess::BaseDir()).
-		Replace(TEXT("{EngineDir}"), *FPaths::EngineDir()).
-		Replace(TEXT("{ProjectDir}"), *FPaths::ProjectDir());
-}
-
-void FGenericPlatformProcess::ModifyCreateProcParams(FString& InOutURL, FString& InOutParams, ECreateProcHelperFlags Flags)
-{
-	FString Command = ReplaceProcHelperVariables(InOutURL);
-	FString Arguments = ReplaceProcHelperVariables(InOutParams);
-
-	if ((Flags & ECreateProcHelperFlags::AppendPlatformScriptExtension) != ECreateProcHelperFlags::None)
-	{
-#if PLATFORM_WINDOWS
-		Command += TEXT(".bat");
-#elif PLATFORM_MAC || PLATFORM_LINUX
-		Command += TEXT(".sh");
-#else
-		UE_LOG(LogInit, Fatal, TEXT("ModifyCreateProcParams being called on a platform that is not supported"));
-#endif
-	}
-
-	Command = FPaths::ConvertRelativePathToFull(Command);
-
-	// wrap in quotes if needed
-	if (Command.Contains(TEXT(" ")))
-	{
-		Command = FString::Printf(TEXT("\"%s\""), *Command);
-	}
-
-	// Mac and Linux need to run .exe files through Mono
-#if PLATFORM_MAC || PLATFORM_LINUX
-	if (Command.EndsWith(TEXT(".exe")))
-	{	
-		Command = FPaths::EngineDir() / TEXT("Build/BatchFiles/Mac/RunMono.sh ") + Command;
-		Flags = Flags | ECreateProcHelperFlags::RunThroughShell;
-	}
-#endif
-
-
-
-	if ((Flags & ECreateProcHelperFlags::RunThroughShell) != ECreateProcHelperFlags::None)
-	{
-#if PLATFORM_WINDOWS
-		InOutURL = TEXT("cmd.exe");
-		InOutParams = FString::Printf(TEXT("/c %s %s"), *Command, *Arguments);
-#elif PLATFORM_MAC
-		InOutURL = TEXT("/bin/sh");
-		InOutParams = FString::Printf(TEXT("%s %s"), *Command, *Arguments);
-#elif PLATFORM_LINUX
-		InOutURL = TEXT("/bin/bash");
-		InOutParams = FString::Printf(TEXT("-c '%s %s'"), *Command, *Arguments);
-#endif
-	}
-	else
-	{
-		// otherwise, we have done all we need to do, just return them!
-		InOutURL = Command;
-		InOutParams = Arguments;
-	}
-}
-
-
 FGenericPlatformProcess::EWaitAndForkResult FGenericPlatformProcess::WaitAndFork()
 {
 	UE_LOG(LogHAL, Fatal, TEXT("FGenericPlatformProcess::WaitAndFork not implemented on this platform"));
diff --git a/Engine/Source/Runtime/Core/Public/GenericPlatform/GenericPlatformProcess.h b/Engine/Source/Runtime/Core/Public/GenericPlatform/GenericPlatformProcess.h
--- a/Engine/Source/Runtime/Core/Public/GenericPlatform/GenericPlatformProcess.h
+++ b/Engine/Source/Runtime/Core/Public/GenericPlatform/GenericPlatformProcess.h
@@ -423,36 +423,6 @@ struct CORE_API FGenericPlatformProcess
 	 */
 	static void TerminateProc( FProcHandle & ProcessHandle, bool KillTree = false );
 
-
-	enum class ECreateProcHelperFlags : uint8
-	{
-		None = 0,
-		// If this is set, the command will be run via cmd.exe, sh, or similar
-		RunThroughShell = 1,
-		// If this is set, the Command will have a platform-specific script extension appended (.bat, .sh, etc)
-		AppendPlatformScriptExtension = 2,
-	};
-
-	/**
-	 * Generates the URL and Parms fields for CreateProc (or FMonitoredProcess) usable on any (generally host) platform. Target platforms that need this will likely need to override it.
-	 *
-	 * Flags are used to determine out to modify the passed in InURL and InParams
-	 * This function will replace instances (in Command or Arguments) of a few variables in {} with their values:
-	 *    {Platform} replaced with the name of the Host platform (Windows, Mac, Linux)
-	 *    {BaseDir} replaced with path to root directory (one up from Engine)
-	 *    {EngineDir} replaced with path to engine directory
-	 *    {ProjectDir} replaced with path to project directory
-	 *
-	 * Ex: CreateProcHelper("{BaseDir}/RunUAT", "MyScript -param", ECreateProcHelperFlags::RunThroughShell | ECreateProcHelperFlags::AppendPlatformScriptExtension)
-	 * Would run "cmd.exe /c RunUAT.bat" on Windows and "sh -c RunUAT.sh" on Mac and Linux
-	 *
-	 * @param InOutURL URL that will be modified and passed to CreateProc or FMonitoredProcess or the like
-	 * @param InOutParams Params to the program that will be modified
-	 * @param Flags to control how the strings are modified
-	 *
-	 */
-	static void ModifyCreateProcParams(FString& InOutURL, FString& InOutParams, ECreateProcHelperFlags Flags);
-
 	enum class EWaitAndForkResult : uint8
 	{
 		Error,
@@ -733,5 +703,3 @@ struct CORE_API FGenericPlatformProcess
 #endif
 	}
 };
-
-ENUM_CLASS_FLAGS(FGenericPlatformProcess::ECreateProcHelperFlags);
