Fixed duplicate -ScriptsForProject command line parameter causing editor turnkey platform SDK detection to fail

diff --git a/Engine/Source/Editor/TurnkeySupport/Private/TurnkeySupportModule.cpp b/Engine/Source/Editor/TurnkeySupport/Private/TurnkeySupportModule.cpp
--- a/Engine/Source/Editor/TurnkeySupport/Private/TurnkeySupportModule.cpp
+++ b/Engine/Source/Editor/TurnkeySupport/Private/TurnkeySupportModule.cpp
@@ -1460,7 +1460,7 @@ static void PrepForTurnkeyReport(FString& Command, FString& BaseCommandline, FSt
 		BaseCommandline.Appendf(TEXT(" -ScriptsForProject=\"%s\" "), *ProjectPath);
 	}
 
-	BaseCommandline = BaseCommandline.Appendf(TEXT("-ScriptsForProject=\"%s\" Turnkey -utf8output -WaitForUATMutex -command=VerifySdk -ReportFilename=\"%s\" -log=\"%s\""), *ProjectPath, *ReportFilename, *LogFilename);
+	BaseCommandline = BaseCommandline.Appendf(TEXT("Turnkey -utf8output -WaitForUATMutex -command=VerifySdk -ReportFilename=\"%s\" -log=\"%s\""), *ReportFilename, *LogFilename);
 }
 
 bool GetSdkInfoFromTurnkey(FString Line, FName& PlatformName, FString& DeviceId, FTurnkeySdkInfo& SdkInfo)
