Fix Mac/Linux crash on start-up.

#rnx

diff --git a/Engine/Source/Editor/TurnkeySupport/Private/TurnkeySupportModule.cpp b/Engine/Source/Editor/TurnkeySupport/Private/TurnkeySupportModule.cpp
--- a/Engine/Source/Editor/TurnkeySupport/Private/TurnkeySupportModule.cpp
+++ b/Engine/Source/Editor/TurnkeySupport/Private/TurnkeySupportModule.cpp
@@ -1442,16 +1442,25 @@ static void PrepForTurnkeyReport(FString& Command, FString& BaseCommandline, FSt
 	// make sure intermediate directory exists
 	IFileManager::Get().MakeDirectory(*FPaths::ProjectIntermediateDir());
 
+#if PLATFORM_WINDOWS
+	FString RunUATScriptName = TEXT("RunUAT.bat");
+#elif PLATFORM_MAC
+	FString RunUATScriptName = TEXT("RunUAT.command");
+#elif PLATFORM_LINUX
+	FString RunUATScriptName = TEXT("RunUAT.sh");
+#endif
+
+	FString UatPath = FPaths::ConvertRelativePathToFull(FPaths::EngineDir() / TEXT("Build/BatchFiles") / RunUATScriptName);
 
 #if PLATFORM_WINDOWS
 	Command = TEXT("cmd.exe");
-	BaseCommandline = TEXT("/c RunUAT.bat");
+	BaseCommandline = FString::Printf(TEXT("/c \"%s\""), *UatPath);
 #elif PLATFORM_MAC
 	Command = TEXT("/bin/sh");
-	BaseCommandline = TEXT("RunUAT.sh");
+	BaseCommandline = FString::Printf(TEXT("\"%s\""), *UatPath);
 #elif PLATFORM_LINUX
 	Command = TEXT("/bin/bash");
-	BaseCommandline = TEXT("-c RunUAT.sh");
+	BaseCommandline = FString::Printf(TEXT("-c \"%s\""), *UatPath)
 #endif
 
 	const FString ProjectPath = GetProjectPathForTurnkey();
@@ -1599,6 +1608,7 @@ private:
 };
 bool FSerializedUATProcess::bHasSucceededOnce = false;
 FCriticalSection FSerializedUATProcess::Serializer;
+static constexpr bool bDeleteTurnkeyProcessOnCompletion = (PLATFORM_WINDOWS);
 
 
 void FTurnkeySupportModule::UpdateSdkInfo()
@@ -1639,7 +1649,7 @@ void FTurnkeySupportModule::UpdateSdkInfo()
 		ClearDeviceStatus();
 	}
 
-	FSerializedUATProcess* TurnkeyProcess = new FSerializedUATProcess(Command, Commandline, true, FOnSerializedProcessComplete::CreateLambda([this, ReportFilename](int32 ExitCode, bool bLaunchFailed, bool bWasCanceled)
+	FSerializedUATProcess* TurnkeyProcess = new FSerializedUATProcess(Command, Commandline, bDeleteTurnkeyProcessOnCompletion, FOnSerializedProcessComplete::CreateLambda([this, ReportFilename](int32 ExitCode, bool bLaunchFailed, bool bWasCanceled)
 	{
 		UE_LOG(LogTurnkeySupport, Log, TEXT("Completed SDK detection: Code = %d, bLaunchFailed = %d, bWasCanceled = %d"), ExitCode, bLaunchFailed, bWasCanceled);
 
@@ -1753,7 +1763,7 @@ void FTurnkeySupportModule::UpdateSdkInfoForDevices(TArray<FString> PlatformDevi
 		}
 	}
 
-	FSerializedUATProcess* TurnkeyProcess = new FSerializedUATProcess(Command, Commandline, true, FOnSerializedProcessComplete::CreateLambda([this, ReportFilename, PlatformDeviceIds](int32 ExitCode, bool bLaunchFailed, bool bWasCanceled)
+	FSerializedUATProcess* TurnkeyProcess = new FSerializedUATProcess(Command, Commandline, bDeleteTurnkeyProcessOnCompletion, FOnSerializedProcessComplete::CreateLambda([this, ReportFilename, PlatformDeviceIds](int32 ExitCode, bool bLaunchFailed, bool bWasCanceled)
 	{
 		UE_LOG(LogTurnkeySupport, Log, TEXT("Completed device detection: Code = %d, bLaunchFailed = %d, bWasCanceled = %d"), ExitCode, bLaunchFailed, bWasCanceled);
 
