Fix Linux CIS

diff --git a/Engine/Source/Editor/TurnkeySupport/Private/TurnkeySupportModule.cpp b/Engine/Source/Editor/TurnkeySupport/Private/TurnkeySupportModule.cpp
--- a/Engine/Source/Editor/TurnkeySupport/Private/TurnkeySupportModule.cpp
+++ b/Engine/Source/Editor/TurnkeySupport/Private/TurnkeySupportModule.cpp
@@ -1460,7 +1460,7 @@ static void PrepForTurnkeyReport(FString& Command, FString& BaseCommandline, FSt
 	BaseCommandline = FString::Printf(TEXT("\"%s\""), *UatPath);
 #elif PLATFORM_LINUX
 	Command = TEXT("/bin/bash");
-	BaseCommandline = FString::Printf(TEXT("-c \"%s\""), *UatPath)
+	BaseCommandline = FString::Printf(TEXT("-c \"%s\""), *UatPath);
 #endif
 
 	const FString ProjectPath = GetProjectPathForTurnkey();
