Fix broken turnkey detection due to cmd /c quote handling

diff --git a/Engine/Source/Editor/TurnkeySupport/Private/TurnkeySupportModule.cpp b/Engine/Source/Editor/TurnkeySupport/Private/TurnkeySupportModule.cpp
--- a/Engine/Source/Editor/TurnkeySupport/Private/TurnkeySupportModule.cpp
+++ b/Engine/Source/Editor/TurnkeySupport/Private/TurnkeySupportModule.cpp
@@ -1432,7 +1432,7 @@ void FTurnkeySupportModule::MakeTurnkeyMenu(FToolMenuSection& MenuSection) const
 
 
 // some shared functionality
-static void PrepForTurnkeyReport(FString& Command, FString& BaseCommandline, FString& ReportFilename)
+static void PrepForTurnkeyReport(FString& BaseCommandline, FString& ReportFilename)
 {
 	static int ReportIndex = 0;
 
@@ -1451,17 +1451,7 @@ static void PrepForTurnkeyReport(FString& Command, FString& BaseCommandline, FSt
 #endif
 
 	FString UatPath = FPaths::ConvertRelativePathToFull(FPaths::EngineDir() / TEXT("Build/BatchFiles") / RunUATScriptName);
-
-#if PLATFORM_WINDOWS
-	Command = TEXT("cmd.exe");
-	BaseCommandline = FString::Printf(TEXT("/c \"%s\""), *UatPath);
-#elif PLATFORM_MAC
-	Command = TEXT("/bin/sh");
 	BaseCommandline = FString::Printf(TEXT("\"%s\""), *UatPath);
-#elif PLATFORM_LINUX
-	Command = TEXT("/bin/bash");
-	BaseCommandline = FString::Printf(TEXT("-c \"%s\""), *UatPath);
-#endif
 
 	const FString ProjectPath = GetProjectPathForTurnkey();
 	if (!ProjectPath.IsEmpty())
@@ -1533,19 +1523,37 @@ DECLARE_DELEGATE_ThreeParams(FOnSerializedProcessComplete, int32 /*ExitCode*/, b
 class FSerializedUATProcess
 {
 public:
-	FSerializedUATProcess(const FString& URL, const FString& Params, bool bDeleteOnCompletion, FOnSerializedProcessComplete CompletionDelegate)
+	FSerializedUATProcess(const FString& Params, bool bDeleteOnCompletion, FOnSerializedProcessComplete CompletionDelegate)
 		: Process(nullptr)
 	{
-		Async(EAsyncExecution::Thread, [this, URL, Params, bDeleteOnCompletion, CompletionDelegate]()
+		Async(EAsyncExecution::Thread, [this, Params, bDeleteOnCompletion, CompletionDelegate]()
 			{
 				// don't let this do anything if another process is running
 				Serializer.Lock();
 
-				FString FinalParams = Params;
+				FString FinalParams;
+#if PLATFORM_WINDOWS
+				FinalParams = TEXT("/c \"");
+#elif PLATFORM_LINUX
+				FinalParams = TEXT("-c ");
+#endif			
+				FinalParams += Params;
 				if (bHasSucceededOnce)
 				{
 					FinalParams += TEXT(" -nocompile");
 				}
+#if PLATFORM_WINDOWS
+				// On Windows, cmd /c needs a set of quotes around the whole command line including parameters
+				FinalParams += TEXT("\"");	
+#endif
+
+#if PLATFORM_WINDOWS
+				FString URL = TEXT("cmd.exe");
+#elif PLATFORM_MAC
+				FString URL = TEXT("/bin/sh");
+#elif PLATFORM_LINUX
+				FString URL = TEXT("/bin/bash");
+#endif
 
 				Process = new FMonitoredProcess(URL, FinalParams, true, false);
 				Process->OnCompleted().BindLambda([this, bDeleteOnCompletion, CompletionDelegate](int32 ExitCode)
@@ -1629,12 +1637,12 @@ void FTurnkeySupportModule::UpdateSdkInfo()
 	}
 
 
-	FString Command, BaseCommandline, ReportFilename;
-	PrepForTurnkeyReport(Command, BaseCommandline, ReportFilename);
+	FString BaseCommandline, ReportFilename;
+	PrepForTurnkeyReport(BaseCommandline, ReportFilename);
 	// get status for all platforms
 	FString Commandline = BaseCommandline + TEXT(" -platform=all");
 
-	UE_LOG(LogTurnkeySupport, Log, TEXT("Running Turnkey SDK detection: '%s %s'"), *Command, *Commandline);
+	UE_LOG(LogTurnkeySupport, Log, TEXT("Running Turnkey SDK detection: '%s'"), *Commandline);
 
 	{
 		FScopeLock Lock(&GTurnkeySection);
@@ -1649,7 +1657,7 @@ void FTurnkeySupportModule::UpdateSdkInfo()
 		ClearDeviceStatus();
 	}
 
-	FSerializedUATProcess* TurnkeyProcess = new FSerializedUATProcess(Command, Commandline, bDeleteTurnkeyProcessOnCompletion, FOnSerializedProcessComplete::CreateLambda([this, ReportFilename](int32 ExitCode, bool bLaunchFailed, bool bWasCanceled)
+	FSerializedUATProcess* TurnkeyProcess = new FSerializedUATProcess(Commandline, bDeleteTurnkeyProcessOnCompletion, FOnSerializedProcessComplete::CreateLambda([this, ReportFilename](int32 ExitCode, bool bLaunchFailed, bool bWasCanceled)
 	{
 		UE_LOG(LogTurnkeySupport, Log, TEXT("Completed SDK detection: Code = %d, bLaunchFailed = %d, bWasCanceled = %d"), ExitCode, bLaunchFailed, bWasCanceled);
 
@@ -1742,14 +1750,14 @@ void FTurnkeySupportModule::UpdateSdkInfo()
 
 void FTurnkeySupportModule::UpdateSdkInfoForDevices(TArray<FString> PlatformDeviceIds)
 {
-	FString Command, BaseCommandline, ReportFilename;
-	PrepForTurnkeyReport(Command, BaseCommandline, ReportFilename);
+	FString BaseCommandline, ReportFilename;
+	PrepForTurnkeyReport(BaseCommandline, ReportFilename);
 
 	// the platform part of the Id may need to be converted to be turnkey (ie UBT) proper
 
 	FString Commandline = BaseCommandline + FString(TEXT(" -Device=")) + FString::JoinBy(PlatformDeviceIds, TEXT("+"), [](FString Id) { return ConvertToUATDeviceId(Id); });
 
-	UE_LOG(LogTurnkeySupport, Log, TEXT("Running Turnkey device detection: '%s %s'"), *Command, *Commandline);
+	UE_LOG(LogTurnkeySupport, Log, TEXT("Running Turnkey device detection: '%s'"), *Commandline);
 
 	{
 		FScopeLock Lock(&GTurnkeySection);
@@ -1763,7 +1771,7 @@ void FTurnkeySupportModule::UpdateSdkInfoForDevices(TArray<FString> PlatformDevi
 		}
 	}
 
-	FSerializedUATProcess* TurnkeyProcess = new FSerializedUATProcess(Command, Commandline, bDeleteTurnkeyProcessOnCompletion, FOnSerializedProcessComplete::CreateLambda([this, ReportFilename, PlatformDeviceIds](int32 ExitCode, bool bLaunchFailed, bool bWasCanceled)
+	FSerializedUATProcess* TurnkeyProcess = new FSerializedUATProcess(Commandline, bDeleteTurnkeyProcessOnCompletion, FOnSerializedProcessComplete::CreateLambda([this, ReportFilename, PlatformDeviceIds](int32 ExitCode, bool bLaunchFailed, bool bWasCanceled)
 	{
 		UE_LOG(LogTurnkeySupport, Log, TEXT("Completed device detection: Code = %d, bLaunchFailed = %d, bWasCanceled = %d"), ExitCode, bLaunchFailed, bWasCanceled);
 
