Pass the project path to Turnkey when launching a map.

diff --git a/Engine/Source/Editor/TurnkeySupport/Private/TurnkeyEditorSupport.cpp b/Engine/Source/Editor/TurnkeySupport/Private/TurnkeyEditorSupport.cpp
--- a/Engine/Source/Editor/TurnkeySupport/Private/TurnkeyEditorSupport.cpp
+++ b/Engine/Source/Editor/TurnkeySupport/Private/TurnkeyEditorSupport.cpp
@@ -80,7 +80,7 @@ void FTurnkeyEditorSupport::PrepareToLaunchRunningMap(const FString& DeviceId, c
 #endif
 }
 
-void FTurnkeyEditorSupport::LaunchRunningMap(const FString& DeviceId, const FString& DeviceName, bool bUseTurnkey)
+void FTurnkeyEditorSupport::LaunchRunningMap(const FString& DeviceId, const FString& DeviceName, const FString& ProjectPath, bool bUseTurnkey)
 {
 #if WITH_EDITOR
 	FTargetDeviceId TargetDeviceId;
@@ -119,7 +119,11 @@ void FTurnkeyEditorSupport::LaunchRunningMap(const FString& DeviceId, const FStr
 			// if we want to check device flash before we start cooking, kick it off now. we could delay this 
 			if (bUseTurnkey)
 			{
-				FString CommandLine = FString::Printf(TEXT("Turnkey -command=VerifySdk -UpdateIfNeeded -platform=%s -EditorIO -noturnkeyvariables -device=%s -utf8output -WaitForUATMutex"), *UBTPlatformName, *TargetDeviceId.GetDeviceName());
+				FString CommandLine = FString::Printf(TEXT("Turnkey -command=VerifySdk -UpdateIfNeeded -platform=%s -EditorIO -noturnkeyvariables -device=%s -utf8output -WaitForUATMutex"), *UBTPlatformName, *ProjectPath, *TargetDeviceId.GetDeviceName());
+				if (!ProjectPath.IsEmpty())
+				{
+					CommandLine += " -project=" + ProjectPath;
+				}
 				FText TaskName = LOCTEXT("VerifyingSDK", "Verifying SDK and Device");
 
 				IUATHelperModule::Get().CreateUatTask(CommandLine, FText::FromString(IniPlatformName), TaskName, TaskName, FEditorStyle::GetBrush(TEXT("MainFrame.PackageProject")),
diff --git a/Engine/Source/Editor/TurnkeySupport/Private/TurnkeyEditorSupport.h b/Engine/Source/Editor/TurnkeySupport/Private/TurnkeyEditorSupport.h
--- a/Engine/Source/Editor/TurnkeySupport/Private/TurnkeyEditorSupport.h
+++ b/Engine/Source/Editor/TurnkeySupport/Private/TurnkeyEditorSupport.h
@@ -12,7 +12,7 @@ public:
 	static FString GetUATOptions();
 
 	static void PrepareToLaunchRunningMap(const FString& DeviceId, const FString& DeviceName);
-	static void LaunchRunningMap(const FString& DeviceId, const FString& DeviceName, bool bUseTurnkey);
+	static void LaunchRunningMap(const FString& DeviceId, const FString& DeviceName, const FString& ProjectPath, bool bUseTurnkey);
 	static void AddEditorOptions(struct FToolMenuSection& MenuBuilder);
 
 	static bool DoesProjectHaveCode();
diff --git a/Engine/Source/Editor/TurnkeySupport/Private/TurnkeySupportModule.cpp b/Engine/Source/Editor/TurnkeySupport/Private/TurnkeySupportModule.cpp
--- a/Engine/Source/Editor/TurnkeySupport/Private/TurnkeySupportModule.cpp
+++ b/Engine/Source/Editor/TurnkeySupport/Private/TurnkeySupportModule.cpp
@@ -1071,7 +1071,7 @@ bool CanLaunchOnDevice(const FString& DeviceName)
 
 static void LaunchOnDevice(const FString& DeviceId, const FString& DeviceName, bool bUseTurnkey)
 {
-	FTurnkeyEditorSupport::LaunchRunningMap(DeviceId, DeviceName, bUseTurnkey);
+	FTurnkeyEditorSupport::LaunchRunningMap(DeviceId, DeviceName, GetProjectPathForTurnkey(), bUseTurnkey);
 }
 
 static void PrepareLaunchOn(FString DeviceId, FString DeviceName)
