Fix a leftover argument setting the device as a project path for Turnkey

diff --git a/Engine/Source/Editor/TurnkeySupport/Private/TurnkeyEditorSupport.cpp b/Engine/Source/Editor/TurnkeySupport/Private/TurnkeyEditorSupport.cpp
--- a/Engine/Source/Editor/TurnkeySupport/Private/TurnkeyEditorSupport.cpp
+++ b/Engine/Source/Editor/TurnkeySupport/Private/TurnkeyEditorSupport.cpp
@@ -119,7 +119,7 @@ void FTurnkeyEditorSupport::LaunchRunningMap(const FString& DeviceId, const FStr
 			// if we want to check device flash before we start cooking, kick it off now. we could delay this 
 			if (bUseTurnkey)
 			{
-				FString CommandLine = FString::Printf(TEXT("Turnkey -command=VerifySdk -UpdateIfNeeded -platform=%s -EditorIO -noturnkeyvariables -device=%s -utf8output -WaitForUATMutex"), *UBTPlatformName, *ProjectPath, *TargetDeviceId.GetDeviceName());
+				FString CommandLine = FString::Printf(TEXT("Turnkey -command=VerifySdk -UpdateIfNeeded -platform=%s -EditorIO -noturnkeyvariables -device=%s -utf8output -WaitForUATMutex"), *UBTPlatformName, *TargetDeviceId.GetDeviceName());
 				if (!ProjectPath.IsEmpty())
 				{
 					CommandLine += " -project=" + ProjectPath;
