- Pass -project to Turnkey when running initial VerifySDK during editor startup

diff --git a/Engine/Source/Editor/TurnkeySupport/Private/TurnkeySupportModule.cpp b/Engine/Source/Editor/TurnkeySupport/Private/TurnkeySupportModule.cpp
--- a/Engine/Source/Editor/TurnkeySupport/Private/TurnkeySupportModule.cpp
+++ b/Engine/Source/Editor/TurnkeySupport/Private/TurnkeySupportModule.cpp
@@ -1454,12 +1454,18 @@ static void PrepForTurnkeyReport(FString& BaseCommandline, FString& ReportFilena
 	BaseCommandline = FString::Printf(TEXT("\"%s\""), *UatPath);
 
 	const FString ProjectPath = GetProjectPathForTurnkey();
+	// now pass a project to UAT
 	if (!ProjectPath.IsEmpty())
 	{
 		BaseCommandline.Appendf(TEXT(" -ScriptsForProject=\"%s\" "), *ProjectPath);
 	}
 
 	BaseCommandline = BaseCommandline.Appendf(TEXT("Turnkey -utf8output -WaitForUATMutex -command=VerifySdk -ReportFilename=\"%s\" -log=\"%s\""), *ReportFilename, *LogFilename);
+	// now pass a project to Turnkey
+	if (!ProjectPath.IsEmpty())
+	{
+		BaseCommandline.Appendf(TEXT(" -project=\"%s\" "), *ProjectPath);
+	}
 }
 
 bool GetSdkInfoFromTurnkey(FString Line, FName& PlatformName, FString& DeviceId, FTurnkeySdkInfo& SdkInfo)
