Add quotes around the project path.

diff --git a/Engine/Source/Editor/TurnkeySupport/Private/TurnkeyEditorSupport.cpp b/Engine/Source/Editor/TurnkeySupport/Private/TurnkeyEditorSupport.cpp
--- a/Engine/Source/Editor/TurnkeySupport/Private/TurnkeyEditorSupport.cpp
+++ b/Engine/Source/Editor/TurnkeySupport/Private/TurnkeyEditorSupport.cpp
@@ -122,7 +122,7 @@ void FTurnkeyEditorSupport::LaunchRunningMap(const FString& DeviceId, const FStr
 				FString CommandLine = FString::Printf(TEXT("Turnkey -command=VerifySdk -UpdateIfNeeded -platform=%s -EditorIO -noturnkeyvariables -device=%s -utf8output -WaitForUATMutex"), *UBTPlatformName, *TargetDeviceId.GetDeviceName());
 				if (!ProjectPath.IsEmpty())
 				{
-					CommandLine += " -project=" + ProjectPath;
+					CommandLine += FString::Printf(TEXT(" -project=\"%s\""), *ProjectPath);
 				}
 				FText TaskName = LOCTEXT("VerifyingSDK", "Verifying SDK and Device");
 
