- Changed FSerializedUATTask to subclass from FMonitoredProcess, and moved it to MonitoredProcess.h/cpp
- Changed UATTask to use a FSerializedUATTask

diff --git a/Engine/Source/Editor/TurnkeySupport/Private/TurnkeySupportModule.cpp b/Engine/Source/Editor/TurnkeySupport/Private/TurnkeySupportModule.cpp
--- a/Engine/Source/Editor/TurnkeySupport/Private/TurnkeySupportModule.cpp
+++ b/Engine/Source/Editor/TurnkeySupport/Private/TurnkeySupportModule.cpp
@@ -1428,31 +1428,22 @@ void FTurnkeySupportModule::MakeTurnkeyMenu(FToolMenuSection& MenuSection) const
 	MenuSection.AddEntry(Entry);
 }
 
+static FString GetLogAndReportCommandline(FString& LogFilename, FString& ReportFilename)
+{
+	static int ReportIndex = 0;
 
+	LogFilename = FPaths::ConvertRelativePathToFull(FPaths::Combine(FPaths::ProjectIntermediateDir(), *FString::Printf(TEXT("TurnkeyLog_%d.log"), ReportIndex)));
+	ReportFilename = FPaths::ConvertRelativePathToFull(FPaths::Combine(FPaths::ProjectIntermediateDir(), *FString::Printf(TEXT("TurnkeyReport_%d.log"), ReportIndex++)));
 
+	return FString::Printf(TEXT("-ReportFilename=\"%s\" -log=\"%s\""), *ReportFilename, *LogFilename);
+}
 
 // some shared functionality
 static void PrepForTurnkeyReport(FString& BaseCommandline, FString& ReportFilename)
 {
-	static int ReportIndex = 0;
-
-	FString LogFilename = FPaths::ConvertRelativePathToFull(FPaths::Combine(FPaths::ProjectIntermediateDir(), *FString::Printf(TEXT("TurnkeyLog_%d.log"), ReportIndex)));
-	ReportFilename = FPaths::ConvertRelativePathToFull(FPaths::Combine(FPaths::ProjectIntermediateDir(), *FString::Printf(TEXT("TurnkeyReport_%d.log"), ReportIndex++)));
-
 	// make sure intermediate directory exists
 	IFileManager::Get().MakeDirectory(*FPaths::ProjectIntermediateDir());
 
-#if PLATFORM_WINDOWS
-	FString RunUATScriptName = TEXT("RunUAT.bat");
-#elif PLATFORM_MAC
-	FString RunUATScriptName = TEXT("RunUAT.command");
-#elif PLATFORM_LINUX
-	FString RunUATScriptName = TEXT("RunUAT.sh");
-#endif
-
-	FString UatPath = FPaths::ConvertRelativePathToFull(FPaths::EngineDir() / TEXT("Build/BatchFiles") / RunUATScriptName);
-	BaseCommandline = FString::Printf(TEXT("\"%s\""), *UatPath);
-
 	const FString ProjectPath = GetProjectPathForTurnkey();
 	// now pass a project to UAT
 	if (!ProjectPath.IsEmpty())
@@ -1460,7 +1451,10 @@ static void PrepForTurnkeyReport(FString& BaseCommandline, FString& ReportFilena
 		BaseCommandline.Appendf(TEXT(" -ScriptsForProject=\"%s\" "), *ProjectPath);
 	}
 
-	BaseCommandline = BaseCommandline.Appendf(TEXT("Turnkey -utf8output -WaitForUATMutex -command=VerifySdk -ReportFilename=\"%s\" -log=\"%s\""), *ReportFilename, *LogFilename);
+	FString LogFilename;
+	FString LogAndReportParams = GetLogAndReportCommandline(LogFilename, ReportFilename);
+
+	BaseCommandline = BaseCommandline.Appendf(TEXT("Turnkey -utf8output -WaitForUATMutex -command=VerifySdk %s"), *LogAndReportParams);
 	// now pass a project to Turnkey
 	if (!ProjectPath.IsEmpty())
 	{
@@ -1524,104 +1518,7 @@ bool GetSdkInfoFromTurnkey(FString Line, FName& PlatformName, FString& DeviceId,
 	return true;
 }
 
-DECLARE_DELEGATE_ThreeParams(FOnSerializedProcessComplete, int32 /*ExitCode*/, bool /*bLaunchFailed*/, bool /*bWasCanceled*/);
-
-class FSerializedUATProcess
-{
-public:
-	FSerializedUATProcess(const FString& Params, bool bDeleteOnCompletion, FOnSerializedProcessComplete CompletionDelegate)
-		: Process(nullptr)
-	{
-		Async(EAsyncExecution::Thread, [this, Params, bDeleteOnCompletion, CompletionDelegate]()
-			{
-				// don't let this do anything if another process is running
-				Serializer.Lock();
-
-				FString FinalParams;
-#if PLATFORM_WINDOWS
-				FinalParams = TEXT("/c \"");
-#elif PLATFORM_LINUX
-				FinalParams = TEXT("-c ");
-#endif			
-				FinalParams += Params;
-				if (bHasSucceededOnce)
-				{
-					FinalParams += TEXT(" -nocompile");
-				}
-#if PLATFORM_WINDOWS
-				// On Windows, cmd /c needs a set of quotes around the whole command line including parameters
-				FinalParams += TEXT("\"");	
-#endif
-
-#if PLATFORM_WINDOWS
-				FString URL = TEXT("cmd.exe");
-#elif PLATFORM_MAC
-				FString URL = TEXT("/bin/sh");
-#elif PLATFORM_LINUX
-				FString URL = TEXT("/bin/bash");
-#endif
-
-				Process = new FMonitoredProcess(URL, FinalParams, true, false);
-				Process->OnCompleted().BindLambda([this, bDeleteOnCompletion, CompletionDelegate](int32 ExitCode)
-					{
-						if (ExitCode == 0 || ExitCode == 10)
-						{
-							bHasSucceededOnce = true;
-						}
-						CompletionDelegate.Execute(ExitCode, false, false);
-
-						// let another one in
-						Serializer.Unlock();
-
-						if (bDeleteOnCompletion)
-						{
-							delete this;
-						}
-					});
-				Process->OnCanceled().BindLambda([this, bDeleteOnCompletion, CompletionDelegate]()
-					{
-						CompletionDelegate.Execute(-1, false, true);
-
-						// let another one in
-						Serializer.Unlock();
-
-						if (bDeleteOnCompletion)
-						{
-							delete this;
-						}
-					});
-				
-				UE_LOG(LogTurnkeySupport, Log, TEXT("Running Serialized UAT: '%s %s'"), *URL, *FinalParams);
-
-				if (Process->Launch() == false)
-				{
-					CompletionDelegate.Execute(-1, true, false);
-
-					// let another one in
-					Serializer.Unlock();
-
-					if (bDeleteOnCompletion)
-					{
-						delete this;
-					}
-				}
-			});
-	}
-
-	~FSerializedUATProcess()
-	{
-		delete Process;
-	}
-
-private:
-
-	static bool bHasSucceededOnce;
-	static FCriticalSection Serializer;
 
-	FMonitoredProcess* Process;
-};
-bool FSerializedUATProcess::bHasSucceededOnce = false;
-FCriticalSection FSerializedUATProcess::Serializer;
 static constexpr bool bDeleteTurnkeyProcessOnCompletion = (PLATFORM_WINDOWS);
 
 
@@ -1663,11 +1560,12 @@ void FTurnkeySupportModule::UpdateSdkInfo()
 		ClearDeviceStatus();
 	}
 
-	FSerializedUATProcess* TurnkeyProcess = new FSerializedUATProcess(Commandline, bDeleteTurnkeyProcessOnCompletion, FOnSerializedProcessComplete::CreateLambda([this, ReportFilename](int32 ExitCode, bool bLaunchFailed, bool bWasCanceled)
+	FSerializedUATProcess* TurnkeyProcess = new FSerializedUATProcess(Commandline, true, false);
+	TurnkeyProcess->OnCompleted().BindLambda([this, ReportFilename, TurnkeyProcess](int32 ExitCode)
 	{
-		UE_LOG(LogTurnkeySupport, Log, TEXT("Completed SDK detection: Code = %d, bLaunchFailed = %d, bWasCanceled = %d"), ExitCode, bLaunchFailed, bWasCanceled);
+		UE_LOG(LogTurnkeySupport, Log, TEXT("Completed SDK detection: ExitCode = %d"), ExitCode);
 
-		AsyncTask(ENamedThreads::GameThread, [this, ReportFilename, ExitCode]()
+		AsyncTask(ENamedThreads::GameThread, [this, ReportFilename, ExitCode, TurnkeyProcess]()
 		{
 			FScopeLock Lock(&GTurnkeySection);
 
@@ -1748,10 +1646,12 @@ void FTurnkeySupportModule::UpdateSdkInfo()
 			}
 
 			// cleanup
-//			delete TurnkeyProcess;
+			delete TurnkeyProcess;
 			IFileManager::Get().Delete(*ReportFilename);
 		});
-	}));
+	});
+
+	TurnkeyProcess->Launch();
 }
 
 void FTurnkeySupportModule::UpdateSdkInfoForDevices(TArray<FString> PlatformDeviceIds)
@@ -1777,11 +1677,12 @@ void FTurnkeySupportModule::UpdateSdkInfoForDevices(TArray<FString> PlatformDevi
 		}
 	}
 
-	FSerializedUATProcess* TurnkeyProcess = new FSerializedUATProcess(Commandline, bDeleteTurnkeyProcessOnCompletion, FOnSerializedProcessComplete::CreateLambda([this, ReportFilename, PlatformDeviceIds](int32 ExitCode, bool bLaunchFailed, bool bWasCanceled)
+	FSerializedUATProcess* TurnkeyProcess = new FSerializedUATProcess(Commandline, true, false);
+	TurnkeyProcess->OnCompleted().BindLambda([this, ReportFilename, PlatformDeviceIds, TurnkeyProcess](int32 ExitCode)
 	{
-		UE_LOG(LogTurnkeySupport, Log, TEXT("Completed device detection: Code = %d, bLaunchFailed = %d, bWasCanceled = %d"), ExitCode, bLaunchFailed, bWasCanceled);
+		UE_LOG(LogTurnkeySupport, Log, TEXT("Completed device detection: Code = %d"), ExitCode);
 
-		AsyncTask(ENamedThreads::GameThread, [this, ReportFilename, PlatformDeviceIds, ExitCode]()
+		AsyncTask(ENamedThreads::GameThread, [this, ReportFilename, PlatformDeviceIds, ExitCode, TurnkeyProcess]()
 		{
 			FScopeLock Lock(&GTurnkeySection);
 
@@ -1834,10 +1735,12 @@ void FTurnkeySupportModule::UpdateSdkInfoForDevices(TArray<FString> PlatformDevi
 			}
 
 			// cleanup
-//			delete TurnkeyProcess;
+			delete TurnkeyProcess;
 			IFileManager::Get().Delete(*ReportFilename);
 		});
-	}));
+	});
+
+	TurnkeyProcess->Launch();
 }
 
 /**
diff --git a/Engine/Source/Editor/UATHelper/UATHelperModule.cpp b/Engine/Source/Editor/UATHelper/UATHelperModule.cpp
--- a/Engine/Source/Editor/UATHelper/UATHelperModule.cpp
+++ b/Engine/Source/Editor/UATHelper/UATHelperModule.cpp
@@ -385,18 +385,6 @@ public:
 
 	virtual void CreateUatTask( const FString& CommandLine, const FText& PlatformDisplayName, const FText& TaskName, const FText &TaskShortName, const FSlateBrush* TaskIcon, UatTaskResultCallack ResultCallback, const FString& ResultLocation)
 	{
-		// make sure that the UAT batch file is in place
-	#if PLATFORM_WINDOWS
-		FString RunUATScriptName = TEXT("RunUAT.bat");
-		FString CmdExe = TEXT("cmd.exe");
-	#elif PLATFORM_LINUX
-		FString RunUATScriptName = TEXT("RunUAT.sh");
-		FString CmdExe = TEXT("/bin/bash");
-	#else
-		FString RunUATScriptName = TEXT("RunUAT.command");
-		FString CmdExe = TEXT("/bin/sh");
-	#endif
-
 		// If this is a packaging or cooking task, clear the PackagingResults log
 		if (!TaskShortName.CompareToCaseIgnored(FText::FromString(TEXT("Packaging"))) || 
 			!TaskShortName.CompareToCaseIgnored(FText::FromString(TEXT("Cooking"))))
@@ -409,14 +397,14 @@ public:
 			}
 		}
 
-		FString UatPath = FPaths::ConvertRelativePathToFull(FPaths::EngineDir() / TEXT("Build/BatchFiles") / RunUATScriptName);
 		FGameProjectGenerationModule& GameProjectModule = FModuleManager::LoadModuleChecked<FGameProjectGenerationModule>(TEXT("GameProjectGeneration"));
 		bool bHasCode = GameProjectModule.Get().ProjectHasCodeFiles();
 		
-		if (!FPaths::FileExists(UatPath))
+		// make sure UAT exists
+		if (!FPaths::FileExists(FSerializedUATProcess::GetUATPath()))
 		{
 			FFormatNamedArguments Arguments;
-			Arguments.Add(TEXT("File"), FText::FromString(UatPath));
+			Arguments.Add(TEXT("File"), FText::FromString(FSerializedUATProcess::GetUATPath()));
 			FMessageDialog::Open(EAppMsgType::Ok, FText::Format(LOCTEXT("RequiredFileNotFoundMessage", "A required file could not be found:\n{File}"), Arguments));
 
 			TArray<FAnalyticsEventAttribute> ParamArray;
@@ -427,13 +415,7 @@ public:
 			return;
 		}
 
-	#if PLATFORM_WINDOWS
-		FString FullCommandLine = FString::Printf(TEXT("/c \"\"%s\" %s\""), *UatPath, *CommandLine);
-	#else
-		FString FullCommandLine = FString::Printf(TEXT("\"%s\" %s"), *UatPath, *CommandLine);
-	#endif
-
-		TSharedPtr<FMonitoredProcess> UatProcess = MakeShareable(new FMonitoredProcess(CmdExe, FullCommandLine, true));
+		TSharedPtr<FSerializedUATProcess> UatProcess = MakeShareable(new FSerializedUATProcess(CommandLine, true));
 
 		FPackagingErrorHandler::ClearAssetErrors();
 
@@ -494,33 +476,13 @@ public:
 		UatProcess->OnCanceled().BindStatic(&FUATHelperModule::HandleUatProcessCanceled, NotificationItemPtr, PlatformDisplayName, TaskShortName, Data);
 		UatProcess->OnCompleted().BindStatic(&FUATHelperModule::HandleUatProcessCompleted, NotificationItemPtr, PlatformDisplayName, TaskShortName, Data, ResultLocation);
 		UatProcess->OnOutput().BindStatic(&FUATHelperModule::HandleUatProcessOutput, NotificationItemPtr, PlatformDisplayName, TaskShortName);
+		UatProcess->OnLaunchFailed().BindStatic(&FUATHelperModule::HandleUatLaunchFailed, NotificationItemPtr, PlatformDisplayName, TaskShortName, Data);
 
-		TWeakPtr<FMonitoredProcess> UatProcessPtr(UatProcess);
+		TWeakPtr<FSerializedUATProcess> UatProcessPtr(UatProcess);
 		FEditorDelegates::OnShutdownPostPackagesSaved.Add(FSimpleDelegate::CreateStatic(&FUATHelperModule::HandleUatCancelButtonClicked, UatProcessPtr));
 
-		if (UatProcess->Launch())
-		{
-			GEditor->PlayEditorSound(TEXT("/Engine/EditorSounds/Notifications/CompileStart_Cue.CompileStart_Cue"));
-		}
-		else
-		{
-			GEditor->PlayEditorSound(TEXT("/Engine/EditorSounds/Notifications/CompileFailed_Cue.CompileFailed_Cue"));
-
-			NotificationItem->SetText(LOCTEXT("UatLaunchFailedNotification", "Failed to launch Unreal Automation Tool (UAT)!"));
-
-			NotificationItem->SetExpireDuration(3.0f);
-			NotificationItem->SetFadeOutDuration(0.5f);
-			NotificationItem->SetCompletionState(SNotificationItem::CS_Fail);
-			NotificationItem->ExpireAndFadeout();
-
-			TArray<FAnalyticsEventAttribute> ParamArray;
-			ParamArray.Add(FAnalyticsEventAttribute(TEXT("Time"), 0.0));
-			FEditorAnalytics::ReportEvent(EventName + TEXT(".Failed"), PlatformDisplayName.ToString(), bHasCode, EAnalyticsErrorCodes::UATLaunchFailure, ParamArray);
-			if (ResultCallback)
-			{
-				ResultCallback(TEXT("FailedToStart"), 0.0f);
-			}
-		}
+		UatProcess->Launch();
+		GEditor->PlayEditorSound(TEXT("/Engine/EditorSounds/Notifications/CompileStart_Cue.CompileStart_Cue"));
 	}
 
 	static void HandleUatHyperlinkNavigate()
@@ -536,7 +498,7 @@ public:
 		}
 	}
 
-	static void HandleUatCancelButtonClicked(TSharedPtr<FMonitoredProcess> PackagerProcess)
+	static void HandleUatCancelButtonClicked(TSharedPtr<FSerializedUATProcess> PackagerProcess)
 	{
 		if ( PackagerProcess.IsValid() )
 		{
@@ -544,7 +506,7 @@ public:
 		}
 	}
 
-	static void HandleUatCancelButtonClicked(TWeakPtr<FMonitoredProcess> PackagerProcessPtr)
+	static void HandleUatCancelButtonClicked(TWeakPtr<FSerializedUATProcess> PackagerProcessPtr)
 	{
 		TSharedPtr<FMonitoredProcess> PackagerProcess = PackagerProcessPtr.Pin();
 		if ( PackagerProcess.IsValid() )
@@ -683,8 +645,25 @@ public:
 				// Deal with any cook errors that may have been encountered.
 				FPackagingErrorHandler::ProcessAndHandleCookErrorOutput(Output);
 			}
+		}
+	}
 
+	static void HandleUatLaunchFailed(TWeakPtr<SNotificationItem> NotificationItemPtr, FText PlatformDisplayName, FText TaskName, EventData Event)
+	{
+		GEditor->PlayEditorSound(TEXT("/Engine/EditorSounds/Notifications/CompileFailed_Cue.CompileFailed_Cue"));
 
+		TGraphTask<FMainFrameActionsNotificationTask>::CreateTask().ConstructAndDispatchWhenReady(
+			NotificationItemPtr,
+			SNotificationItem::CS_Fail,
+			LOCTEXT("UatLaunchFailedNotification", "Failed to launch Unreal Automation Tool (UAT)!")
+		);
+
+		TArray<FAnalyticsEventAttribute> ParamArray;
+		ParamArray.Add(FAnalyticsEventAttribute(TEXT("Time"), 0.0));
+		FEditorAnalytics::ReportEvent(Event.EventName + TEXT(".Failed"), PlatformDisplayName.ToString(), Event.bProjectHasCode, EAnalyticsErrorCodes::UATLaunchFailure, ParamArray);
+		if (Event.ResultCallback)
+		{
+			Event.ResultCallback(TEXT("FailedToStart"), 0.0f);
 		}
 	}
 
diff --git a/Engine/Source/Runtime/Core/Private/Misc/MonitoredProcess.cpp b/Engine/Source/Runtime/Core/Private/Misc/MonitoredProcess.cpp
--- a/Engine/Source/Runtime/Core/Private/Misc/MonitoredProcess.cpp
+++ b/Engine/Source/Runtime/Core/Private/Misc/MonitoredProcess.cpp
@@ -3,8 +3,12 @@
 #include "Misc/MonitoredProcess.h"
 #include "HAL/RunnableThread.h"
 #include "Misc/Paths.h"
+#include "Logging/LogMacros.h"
+#include "Async/Async.h"
 #include <atomic>
 
+DEFINE_LOG_CATEGORY_STATIC(LogMonitoredProcess, Log, All);
+
 /* FMonitoredProcess structors
  *****************************************************************************/
 
@@ -183,3 +187,89 @@ void FMonitoredProcess::Tick()
 	}
 }
 
+
+
+
+FCriticalSection FSerializedUATProcess::Serializer;
+
+FSerializedUATProcess::FSerializedUATProcess(const FString& RunUATCommandline, bool InHidden, bool InCreatePipes)
+	// we will modify URL and Params in this constructor, so there's no need to pass anything up to base
+	: FMonitoredProcess("", "", InHidden, InCreatePipes)
+{
+	// replace the URL and Params freom base with the shelled version
+
+#if PLATFORM_WINDOWS
+	URL = TEXT("cmd.exe");
+	Params = FString::Printf(TEXT("/c \"\"%s\" %s\""), *GetUATPath(), *RunUATCommandline);
+#elif PLATFORM_LINUX
+	URL = TEXT("/bin/bash");
+	Params = FString::Printf(TEXT("\"%s\" %s"), *GetUATPath(), *RunUATCommandline);
+#elif PLATFORM_MAC
+	URL = TEXT("/bin/sh");
+	Params = FString::Printf(TEXT("\"%s\" %s"), *GetUATPath(), *RunUATCommandline);
+#endif
+}
+
+bool FSerializedUATProcess::Launch()
+{
+	Async(EAsyncExecution::Thread, [this]()
+		{
+			// don't let this do anything if another process is running
+			Serializer.Lock();
+
+			if (bHasSucceededOnce)
+			{
+				Params += TEXT(" -nocompile");
+			}
+
+			FOnMonitoredProcessCompleted OriginalCompletedDelegate = CompletedDelegate;
+			FSimpleDelegate OriginalCanceledDelegate = CanceledDelegate;
+
+			CompletedDelegate.BindLambda([this, OriginalCompletedDelegate](int32 ExitCode)
+				{
+					if (ExitCode == 0 || ExitCode == 10)
+					{
+						bHasSucceededOnce = true;
+					}
+					OriginalCompletedDelegate.ExecuteIfBound(ExitCode);
+
+					// let another one in
+					Serializer.Unlock();
+				});
+			CanceledDelegate.BindLambda([this, OriginalCanceledDelegate]()
+				{
+					OriginalCanceledDelegate.ExecuteIfBound();
+
+					// let another one in
+					Serializer.Unlock();
+				});
+
+			UE_LOG(LogMonitoredProcess, Log, TEXT("Running Serialized UAT: '%s %s'"), *URL, *Params);
+
+			if (FMonitoredProcess::Launch() == false)
+			{
+				LaunchFailedDelegate.ExecuteIfBound();
+
+				// let another one in
+				Serializer.Unlock();
+			}
+		});
+
+	return true;
+}
+
+
+bool FSerializedUATProcess::bHasSucceededOnce = false;
+
+FString FSerializedUATProcess::GetUATPath()
+{
+#if PLATFORM_WINDOWS
+	FString RunUATScriptName = TEXT("RunUAT.bat");
+#elif PLATFORM_LINUX
+	FString RunUATScriptName = TEXT("RunUAT.sh");
+#else
+	FString RunUATScriptName = TEXT("RunUAT.command");
+#endif
+
+	return FPaths::ConvertRelativePathToFull(FPaths::EngineDir() / TEXT("Build/BatchFiles") / RunUATScriptName);
+}
diff --git a/Engine/Source/Runtime/Core/Public/Misc/MonitoredProcess.h b/Engine/Source/Runtime/Core/Public/Misc/MonitoredProcess.h
--- a/Engine/Source/Runtime/Core/Public/Misc/MonitoredProcess.h
+++ b/Engine/Source/Runtime/Core/Public/Misc/MonitoredProcess.h
@@ -56,7 +56,7 @@ public:
 	FMonitoredProcess( const FString& InURL, const FString& InParams, const FString& InWorkingDir, bool InHidden, bool InCreatePipes = true );
 
 	/** Destructor. */
-	~FMonitoredProcess();
+	virtual ~FMonitoredProcess();
 
 public:
 
@@ -96,7 +96,7 @@ public:
 	bool Update();
 
 	/** Launches the process. */
-	bool Launch();
+	virtual bool Launch();
 
 	/**
 	 * Sets the sleep interval to be used in the main thread loop.
@@ -187,7 +187,7 @@ protected:
 	 */
 	void ProcessOutput( const FString& Output );
 
-private:
+protected:
 	void TickInternal();
 
 
@@ -242,7 +242,7 @@ private:
 	// Buffered output text which does not contain a newline
 	FString OutputBuffer;
 
-private:
+protected:
 
 	// Holds a delegate that is executed when the process has been canceled. */
 	FSimpleDelegate CanceledDelegate;
@@ -253,3 +253,42 @@ private:
 	// Holds a delegate that is executed when a monitored process produces output. */
 	FOnMonitoredProcessOutput OutputDelegate;
 };
+
+
+class CORE_API FSerializedUATProcess : public FMonitoredProcess
+{
+public:
+	/**
+	 * Get the host-platform-specific path to the UAT running script
+	 */
+	static FString GetUATPath();
+
+public:
+	FSerializedUATProcess(const FString& RunUATCommandline, bool InHidden, bool InCreatePipes=true);
+
+	/**
+	 * Run UAT, serially with other FSerializedUATProcess objects. Because the actual call is delayed, this will
+	 * always return true, and the LaunchFailedDelegate will be called later if an error happens
+	 */
+	virtual bool Launch() override;
+
+	/**
+	 * Returns a delegate that is executed when the process has been canceled.
+	 *
+	 * @return The delegate.
+	 */
+	FSimpleDelegate& OnLaunchFailed()
+	{
+		return LaunchFailedDelegate;
+	}
+
+
+private:
+
+	// Holds a delegate that is executed when the process fails to launch (delayed, in a thread). Used in place of the return value
+	// of Launch in the parent class, since it's async
+	FSimpleDelegate LaunchFailedDelegate;
+
+	static FCriticalSection Serializer;
+	static bool bHasSucceededOnce;
+};
