Pass flag -skipbuildeditor as well as -nocompileeditor from Turnkey

diff --git a/Engine/Source/Editor/TurnkeySupport/Private/TurnkeySupportModule.cpp b/Engine/Source/Editor/TurnkeySupport/Private/TurnkeySupportModule.cpp
--- a/Engine/Source/Editor/TurnkeySupport/Private/TurnkeySupportModule.cpp
+++ b/Engine/Source/Editor/TurnkeySupport/Private/TurnkeySupportModule.cpp
@@ -95,7 +95,7 @@ protected:
 	{
 		// We never want to compile editor targets when invoking UAT in this context.
 		// If we are installed or don't have a compiler, we must assume we have a precompiled UAT.
-		return TEXT("-nocompileeditor");
+		return TEXT("-nocompileeditor -skipbuildeditor");
 	}
 
 	static bool ShowBadSDKDialog(FName IniPlatformName)
