Always create a write/read pipes for MonitoredProcesses, otherwise dotnet can hang

Assuming its blocking on UE consuming the output of dotnet

diff --git a/Engine/Source/Editor/TurnkeySupport/Private/TurnkeySupportModule.cpp b/Engine/Source/Editor/TurnkeySupport/Private/TurnkeySupportModule.cpp
--- a/Engine/Source/Editor/TurnkeySupport/Private/TurnkeySupportModule.cpp
+++ b/Engine/Source/Editor/TurnkeySupport/Private/TurnkeySupportModule.cpp
@@ -1627,7 +1627,7 @@ void FTurnkeySupportModule::UpdateSdkInfo()
 		ClearDeviceStatus();
 	}
 
-	FSerializedUATProcess* TurnkeyProcess = new FSerializedUATProcess(Commandline, true, false);
+	FSerializedUATProcess* TurnkeyProcess = new FSerializedUATProcess(Commandline);
 	TurnkeyProcess->OnCompleted().BindLambda([this, ReportFilename, TurnkeyProcess](int32 ExitCode)
 	{
 		UE_LOG(LogTurnkeySupport, Log, TEXT("Completed SDK detection: ExitCode = %d"), ExitCode);
@@ -1744,7 +1744,7 @@ void FTurnkeySupportModule::UpdateSdkInfoForDevices(TArray<FString> PlatformDevi
 		}
 	}
 
-	FSerializedUATProcess* TurnkeyProcess = new FSerializedUATProcess(Commandline, true, false);
+	FSerializedUATProcess* TurnkeyProcess = new FSerializedUATProcess(Commandline);
 	TurnkeyProcess->OnCompleted().BindLambda([this, ReportFilename, PlatformDeviceIds, TurnkeyProcess](int32 ExitCode)
 	{
 		UE_LOG(LogTurnkeySupport, Log, TEXT("Completed device detection: Code = %d"), ExitCode);
diff --git a/Engine/Source/Editor/UATHelper/UATHelperModule.cpp b/Engine/Source/Editor/UATHelper/UATHelperModule.cpp
--- a/Engine/Source/Editor/UATHelper/UATHelperModule.cpp
+++ b/Engine/Source/Editor/UATHelper/UATHelperModule.cpp
@@ -415,7 +415,7 @@ public:
 			return;
 		}
 
-		TSharedPtr<FSerializedUATProcess> UatProcess = MakeShareable(new FSerializedUATProcess(CommandLine, true));
+		TSharedPtr<FSerializedUATProcess> UatProcess = MakeShareable(new FSerializedUATProcess(CommandLine));
 
 		FPackagingErrorHandler::ClearAssetErrors();
 
diff --git a/Engine/Source/Runtime/Core/Private/Misc/MonitoredProcess.cpp b/Engine/Source/Runtime/Core/Private/Misc/MonitoredProcess.cpp
--- a/Engine/Source/Runtime/Core/Private/Misc/MonitoredProcess.cpp
+++ b/Engine/Source/Runtime/Core/Private/Misc/MonitoredProcess.cpp
@@ -70,7 +70,7 @@ bool FMonitoredProcess::Launch()
 		return false;
 	}
 
-	ProcessHandle = FPlatformProcess::CreateProc(*URL, *Params, false, Hidden, Hidden, nullptr, 0, *WorkingDir, WritePipe);
+	ProcessHandle = FPlatformProcess::CreateProc(*URL, *Params, false, Hidden, Hidden, nullptr, 0, *WorkingDir, WritePipe, ReadPipe);
 
 	if (!ProcessHandle.IsValid())
 	{
@@ -190,9 +190,9 @@ void FMonitoredProcess::Tick()
 
 /* FSerializedUATProcess
 *****************************************************************************/
-FSerializedUATProcess::FSerializedUATProcess(const FString& RunUATCommandline, bool InHidden, bool InCreatePipes)
+FSerializedUATProcess::FSerializedUATProcess(const FString& RunUATCommandline)
 	// we will modify URL and Params in this constructor, so there's no need to pass anything up to base
-	: FMonitoredProcess("", "", InHidden, InCreatePipes)
+	: FMonitoredProcess("", "", true, true)
 {
 	// replace the URL and Params freom base with the shelled version
 
diff --git a/Engine/Source/Runtime/Core/Public/Misc/MonitoredProcess.h b/Engine/Source/Runtime/Core/Public/Misc/MonitoredProcess.h
--- a/Engine/Source/Runtime/Core/Public/Misc/MonitoredProcess.h
+++ b/Engine/Source/Runtime/Core/Public/Misc/MonitoredProcess.h
@@ -264,7 +264,7 @@ public:
 	static FString GetUATPath();
 
 public:
-	FSerializedUATProcess(const FString& RunUATCommandline, bool InHidden, bool InCreatePipes=true);
+	FSerializedUATProcess(const FString& RunUATCommandline);
 
 	/**
 	 * Run UAT, serially with other FSerializedUATProcess objects. Because the actual call is delayed, this will
