Fix issue wth not being able to package when a path to executable contains spaces

diff --git a/Engine/Source/Editor/TurnkeySupport/Private/TurnkeyEditorSupport.cpp b/Engine/Source/Editor/TurnkeySupport/Private/TurnkeyEditorSupport.cpp
--- a/Engine/Source/Editor/TurnkeySupport/Private/TurnkeyEditorSupport.cpp
+++ b/Engine/Source/Editor/TurnkeySupport/Private/TurnkeyEditorSupport.cpp
@@ -29,7 +29,7 @@ FString FTurnkeyEditorSupport::GetUATOptions()
 {
 #if WITH_EDITOR
 	FString Options;
-	Options += FString::Printf(TEXT(" -ue4exe=%s"), *FUnrealEdMisc::Get().GetExecutableForCommandlets());
+	Options += FString::Printf(TEXT(" -ue4exe=\"%s\""), *FUnrealEdMisc::Get().GetExecutableForCommandlets());
 
 	int32 NumCookers = GetDefault<UEditorExperimentalSettings>()->MultiProcessCooking;
 	if (NumCookers > 0)
