Remove editor as a valid target selection for Turnkey

diff --git a/Engine/Source/Editor/TurnkeySupport/Private/TurnkeySupportModule.cpp b/Engine/Source/Editor/TurnkeySupport/Private/TurnkeySupportModule.cpp
--- a/Engine/Source/Editor/TurnkeySupport/Private/TurnkeySupportModule.cpp
+++ b/Engine/Source/Editor/TurnkeySupport/Private/TurnkeySupportModule.cpp
@@ -927,6 +927,12 @@ static void MakeTurnkeyPlatformMenu(FMenuBuilder& MenuBuilder, FName IniPlatform
 
 			for (const PlatformInfo::FTargetPlatformInfo* Info : AllTargets)
 			{
+				// Editor isn't a valid platform type that users can target
+				if(Info->PlatformType == EBuildTargetType::Editor)
+				{
+					continue;
+				}
+
 				MenuBuilder.AddMenuEntry(
 					Info->DisplayName,
 					FText(),
