Remove Target Selection menu from binary builds (there's only one option since Client/Server/Editor can't be targeted)

diff --git a/Engine/Source/Editor/TurnkeySupport/Private/TurnkeySupportModule.cpp b/Engine/Source/Editor/TurnkeySupport/Private/TurnkeySupportModule.cpp
--- a/Engine/Source/Editor/TurnkeySupport/Private/TurnkeySupportModule.cpp
+++ b/Engine/Source/Editor/TurnkeySupport/Private/TurnkeySupportModule.cpp
@@ -251,8 +251,17 @@ public:
 		// get a in-memory defaults which will have the user-settings, like the per-platform config/target platform stuff
 		UProjectPackagingSettings* AllPlatformPackagingSettings = GetMutableDefault<UProjectPackagingSettings>();
 	
-		// get the user-chosen targetplatform
-		const PlatformInfo::FTargetPlatformInfo* PlatformInfo = PlatformInfo::FindPlatformInfo(AllPlatformPackagingSettings->GetTargetPlatformForPlatform(IniPlatformName));
+		// installed builds only support standard Game type builds (not Client, Server, etc) so instead of looking up a setting that the user can't set, 
+		// always use the base PlatformInfo for Game builds, which will be named the same as the platform itself
+		const PlatformInfo::FTargetPlatformInfo* PlatformInfo = nullptr;
+		if(FApp::IsInstalled())
+		{
+			PlatformInfo = PlatformInfo::FindPlatformInfo(IniPlatformName);
+		}
+		else
+		{
+			PlatformInfo = PlatformInfo::FindPlatformInfo(AllPlatformPackagingSettings->GetTargetPlatformForPlatform(IniPlatformName));
+		}
 		// this is unexpected to be able to happen, but it could if there was a bad value saved in the UProjectPackagingSettings - if this trips, we should handle errors
 		check(PlatformInfo != nullptr);
 
@@ -919,20 +928,23 @@ static void MakeTurnkeyPlatformMenu(FMenuBuilder& MenuBuilder, FName IniPlatform
 			}
 		MenuBuilder.EndSection();
 
-		MenuBuilder.BeginSection("TargetSelection", LOCTEXT("TurnkeySection_TargetSelection", "Target Selection"));
-
+		// Binary builds can't compile for Client, Server or Editor so that leaves only the default option. No point creating a section for one option. 
+		if (!FApp::IsEngineInstalled())
+		{
+			MenuBuilder.BeginSection("TargetSelection", LOCTEXT("TurnkeySection_TargetSelection", "Target Selection"));
+			
 			// gather all platform infos
 			TArray<const PlatformInfo::FTargetPlatformInfo*> AllTargets = { VanillaInfo };
 			AllTargets.Append(VanillaInfo->Flavors);
-
+			
 			for (const PlatformInfo::FTargetPlatformInfo* Info : AllTargets)
 			{
 				// Editor isn't a valid platform type that users can target
-				if(Info->PlatformType == EBuildTargetType::Editor)
+				if (Info->PlatformType == EBuildTargetType::Editor)
 				{
 					continue;
 				}
-
+				
 				MenuBuilder.AddMenuEntry(
 					Info->DisplayName,
 					FText(),
@@ -946,7 +958,8 @@ static void MakeTurnkeyPlatformMenu(FMenuBuilder& MenuBuilder, FName IniPlatform
 					EUserInterfaceActionType::RadioButton
 				);
 			}
-		MenuBuilder.EndSection();
+			MenuBuilder.EndSection();
+		}
 
 		MenuBuilder.BeginSection("AllDevices", LOCTEXT("TurnkeySection_AllDevices", "All Devices"));
 
