fixed TurnkeyProcess trying to launch a task after TaskGraph was already shut down

diff --git a/Engine/Source/Editor/TurnkeySupport/Private/TurnkeySupportModule.cpp b/Engine/Source/Editor/TurnkeySupport/Private/TurnkeySupportModule.cpp
--- a/Engine/Source/Editor/TurnkeySupport/Private/TurnkeySupportModule.cpp
+++ b/Engine/Source/Editor/TurnkeySupport/Private/TurnkeySupportModule.cpp
@@ -16,6 +16,7 @@
 #include "ITargetDeviceServicesModule.h"
 #include "Misc/DataDrivenPlatformInfoRegistry.h"
 #include "Async/Async.h"
+#include "HAL/Event.h"
 #include "Misc/FileHelper.h"
 #include "Misc/ConfigCacheIni.h"
 #include "Delegates/Delegate.h"
@@ -1651,7 +1652,15 @@ void FTurnkeySupportModule::UpdateSdkInfo()
 	}
 
 	FSerializedUATProcess* TurnkeyProcess = new FSerializedUATProcess(Commandline);
-	TurnkeyProcess->OnCompleted().BindLambda([this, ReportFilename, TurnkeyProcess](int32 ExitCode)
+
+	// block the game thread on exit while TurnkeyProcess being cancelled
+	struct FEventDeleter { void operator()(FEvent* Event) { FPlatformProcess::ReturnSynchEventToPool(Event); } };
+	TSharedPtr<FEvent, ESPMode::ThreadSafe> Barrier{ FPlatformProcess::GetSynchEventFromPool(), FEventDeleter{} };
+
+	TurnkeyProcess->OnCanceled().BindLambda([Barrier] { Barrier->Trigger(); });
+	FCoreDelegates::OnExit.AddLambda([TurnkeyProcess, Barrier] { TurnkeyProcess->Cancel(); Barrier->Wait(); });
+
+	TurnkeyProcess->OnCompleted().BindLambda([this, ReportFilename, TurnkeyProcess, Barrier](int32 ExitCode)
 	{
 		UE_LOG(LogTurnkeySupport, Log, TEXT("Completed SDK detection: ExitCode = %d"), ExitCode);
 
@@ -1739,6 +1748,8 @@ void FTurnkeySupportModule::UpdateSdkInfo()
 			delete TurnkeyProcess;
 			IFileManager::Get().Delete(*ReportFilename);
 		});
+
+		Barrier->Trigger();
 	});
 
 	TurnkeyProcess->Launch();
