Clean up OnExit delegate when closing the Turnkey Process

diff --git a/Engine/Source/Editor/TurnkeySupport/Private/TurnkeySupportModule.cpp b/Engine/Source/Editor/TurnkeySupport/Private/TurnkeySupportModule.cpp
--- a/Engine/Source/Editor/TurnkeySupport/Private/TurnkeySupportModule.cpp
+++ b/Engine/Source/Editor/TurnkeySupport/Private/TurnkeySupportModule.cpp
@@ -1658,13 +1658,13 @@ void FTurnkeySupportModule::UpdateSdkInfo()
 	TSharedPtr<FEvent, ESPMode::ThreadSafe> Barrier{ FPlatformProcess::GetSynchEventFromPool(), FEventDeleter{} };
 
 	TurnkeyProcess->OnCanceled().BindLambda([Barrier] { Barrier->Trigger(); });
-	FCoreDelegates::OnExit.AddLambda([TurnkeyProcess, Barrier] { TurnkeyProcess->Cancel(); Barrier->Wait(); });
+	FDelegateHandle OnExitHandle = FCoreDelegates::OnExit.AddLambda([TurnkeyProcess, Barrier] { TurnkeyProcess->Cancel(); Barrier->Wait(); });
 
-	TurnkeyProcess->OnCompleted().BindLambda([this, ReportFilename, TurnkeyProcess, Barrier](int32 ExitCode)
+	TurnkeyProcess->OnCompleted().BindLambda([this, ReportFilename, TurnkeyProcess, Barrier, OnExitHandle](int32 ExitCode)
 	{
 		UE_LOG(LogTurnkeySupport, Log, TEXT("Completed SDK detection: ExitCode = %d"), ExitCode);
 
-		AsyncTask(ENamedThreads::GameThread, [this, ReportFilename, ExitCode, TurnkeyProcess]()
+		AsyncTask(ENamedThreads::GameThread, [this, ReportFilename, ExitCode, TurnkeyProcess, OnExitHandle]()
 		{
 			FScopeLock Lock(&GTurnkeySection);
 
@@ -1745,6 +1745,7 @@ void FTurnkeySupportModule::UpdateSdkInfo()
 			}
 
 			// cleanup
+			FCoreDelegates::OnExit.Remove(OnExitHandle);
 			delete TurnkeyProcess;
 			IFileManager::Get().Delete(*ReportFilename);
 		});
