Add back a save of dirty packages like it used to when kicking a packaging or cooking through Turnkey

diff --git a/Engine/Source/Editor/TurnkeySupport/Private/TurnkeyEditorSupport.cpp b/Engine/Source/Editor/TurnkeySupport/Private/TurnkeyEditorSupport.cpp
--- a/Engine/Source/Editor/TurnkeySupport/Private/TurnkeyEditorSupport.cpp
+++ b/Engine/Source/Editor/TurnkeySupport/Private/TurnkeyEditorSupport.cpp
@@ -21,6 +21,7 @@
 #include "Settings/EditorExperimentalSettings.h"
 #include "Interfaces/IMainFrameModule.h"
 #include "ToolMenus.h"
+#include "FileHelpers.h"
 #endif
 
 #define LOCTEXT_NAMESPACE "FTurnkeyEditorSupport"
@@ -161,6 +162,20 @@ void FTurnkeyEditorSupport::LaunchRunningMap(const FString& DeviceId, const FStr
 #endif
 }
 
+void FTurnkeyEditorSupport::SaveAll()
+{
+#if WITH_EDITOR
+	FEditorFileUtils::SaveDirtyPackages(false /*bPromptUserToSave*/
+		, true /*bSaveMapPackages*/
+		, true /*bSaveContentPackages*/
+		, false /*bFastSave*/
+		, false /*bNotifyNoPackagesSaved*/
+		, false /*bCanBeDeclined*/);
+#else
+	unimplemented();
+#endif
+}
+
 bool FTurnkeyEditorSupport::DoesProjectHaveCode()
 {
 #if WITH_EDITOR
diff --git a/Engine/Source/Editor/TurnkeySupport/Private/TurnkeyEditorSupport.h b/Engine/Source/Editor/TurnkeySupport/Private/TurnkeyEditorSupport.h
--- a/Engine/Source/Editor/TurnkeySupport/Private/TurnkeyEditorSupport.h
+++ b/Engine/Source/Editor/TurnkeySupport/Private/TurnkeyEditorSupport.h
@@ -14,7 +14,8 @@ public:
 	static void PrepareToLaunchRunningMap(const FString& DeviceId, const FString& DeviceName);
 	static void LaunchRunningMap(const FString& DeviceId, const FString& DeviceName, const FString& ProjectPath, bool bUseTurnkey);
 	static void AddEditorOptions(struct FToolMenuSection& MenuBuilder);
-
+	
+	static void SaveAll();
 	static bool DoesProjectHaveCode();
 	static void RunUAT(const FString& CommandLine, const FText& PlatformDisplayName, const FText& TaskName, const FText& TaskShortName, const struct FSlateBrush* TaskIcon, TFunction<void(FString, double)> ResultCallback=TFunction<void(FString, double)>());
 
diff --git a/Engine/Source/Editor/TurnkeySupport/Private/TurnkeySupportModule.cpp b/Engine/Source/Editor/TurnkeySupport/Private/TurnkeySupportModule.cpp
--- a/Engine/Source/Editor/TurnkeySupport/Private/TurnkeySupportModule.cpp
+++ b/Engine/Source/Editor/TurnkeySupport/Private/TurnkeySupportModule.cpp
@@ -292,6 +292,8 @@ public:
 			}
 		}
 
+		// force a save of dirty packages before proceeding to run UAT
+		FTurnkeyEditorSupport::SaveAll();
 
 		// basic BuildCookRun params we always want
 		FString BuildCookRunParams = FString::Printf(TEXT("-nop4 -utf8output %s -cook "), GetUATCompilationFlags());
