fix UProjectPackagingSettings being deleted while pointer held in packaging

diff --git a/Engine/Source/Editor/TurnkeySupport/Private/TurnkeySupportModule.cpp b/Engine/Source/Editor/TurnkeySupport/Private/TurnkeySupportModule.cpp
--- a/Engine/Source/Editor/TurnkeySupport/Private/TurnkeySupportModule.cpp
+++ b/Engine/Source/Editor/TurnkeySupport/Private/TurnkeySupportModule.cpp
@@ -247,8 +247,6 @@ public:
 
 	static void CookOrPackage(FName IniPlatformName, EPrepareContentMode Mode)
 	{
-		// get all the helper objects
-		UProjectPackagingSettings* PackagingSettings = GetPackagingSettingsForPlatform(IniPlatformName);
 		// get a in-memory defaults which will have the user-settings, like the per-platform config/target platform stuff
 		UProjectPackagingSettings* AllPlatformPackagingSettings = GetMutableDefault<UProjectPackagingSettings>();
 	
@@ -269,7 +267,6 @@ public:
 		const FString UBTPlatformString = PlatformInfo->DataDrivenPlatformInfo->UBTPlatformString;
 		const FString ProjectPath = GetProjectPathForTurnkey();
 
-
 		// check that we can proceed
 		{
 			if (FInstalledPlatformInfo::Get().IsPlatformMissingRequiredFile(UBTPlatformString))
@@ -293,8 +290,12 @@ public:
 		}
 
 		// force a save of dirty packages before proceeding to run UAT
+		// this may delete UProjectPackagingSettings , don't hold it across this call
 		FTurnkeyEditorSupport::SaveAll();
 
+		// get all the helper objects
+		UProjectPackagingSettings* PackagingSettings = GetPackagingSettingsForPlatform(IniPlatformName);
+
 		// basic BuildCookRun params we always want
 		FString BuildCookRunParams = FString::Printf(TEXT("-nop4 -utf8output %s -cook "), GetUATCompilationFlags());
 
