Restore support for packaging multiple game/client/server targets.
Add a default project setting for the build target and binary configuration selection.
Removed 'platformtarget' menu and exposed 'target' and 'flavor' menus (flavor is only used for android now but will also work if other platforms add their own flavors).
Setting -distribution now forces packaging with shipping.

diff --git a/Engine/Source/Developer/DeveloperToolSettings/Classes/Settings/ProjectPackagingSettings.h b/Engine/Source/Developer/DeveloperToolSettings/Classes/Settings/ProjectPackagingSettings.h
--- a/Engine/Source/Developer/DeveloperToolSettings/Classes/Settings/ProjectPackagingSettings.h
+++ b/Engine/Source/Developer/DeveloperToolSettings/Classes/Settings/ProjectPackagingSettings.h
@@ -543,23 +543,31 @@ public:
 	 */
 	EProjectPackagingBuildConfigurations GetBuildConfigurationForPlatform(FName PlatformName) const;
 	void SetBuildConfigurationForPlatform(FName PlatformName, EProjectPackagingBuildConfigurations Configuration);
+	void SetBuildConfigurationForAllPlatforms(EProjectPackagingBuildConfigurations Configuration);
 
-	FName GetTargetPlatformForPlatform(FName PlatformName) const;
-	void SetTargetPlatformForPlatform(FName PlatformName, FName TargetPlatformName);
+	FName GetTargetFlavorForPlatform(FName PlatformName) const;
+	void SetTargetFlavorForPlatform(FName PlatformName, FName TargetFlavorName);
+
+	FString GetBuildTargetForPlatform(FName PlatformName) const;
+	void SetBuildTargetForPlatform(FName PlatformName, FString BuildTargetName);
 
 private:
 	/** Per platform build configuration */
 	UPROPERTY(config)
 	TMap<FName, EProjectPackagingBuildConfigurations> PerPlatformBuildConfig;
 
-	/** Per platform targetplatform cooking target */
+	/** Per platform flavor cooking target */
+	UPROPERTY(config)
+	TMap<FName, FName> PerPlatformTargetFlavorName;
+
+	/** Per platform build target */
 	UPROPERTY(config)
-	TMap<FName, FName> PerPlatformTargetPlatformName;
+	TMap<FName, FString> PerPlatformBuildTarget;
 
 	/** Helper array used to mirror Blueprint asset selections across edits */
 	TArray<FFilePath> CachedNativizeBlueprintAssets;
 	
-	/** THe platform to LoadConfig for to get platform-specific packaging settings */
+	/** The platform to LoadConfig for to get platform-specific packaging settings */
 	FString ConfigPlatform;
 
 public:
diff --git a/Engine/Source/Developer/DeveloperToolSettings/Private/ProjectPackagingSettings.cpp b/Engine/Source/Developer/DeveloperToolSettings/Private/ProjectPackagingSettings.cpp
--- a/Engine/Source/Developer/DeveloperToolSettings/Private/ProjectPackagingSettings.cpp
+++ b/Engine/Source/Developer/DeveloperToolSettings/Private/ProjectPackagingSettings.cpp
@@ -219,7 +219,8 @@ EProjectPackagingBuildConfigurations UProjectPackagingSettings::GetBuildConfigur
 {
 	const EProjectPackagingBuildConfigurations* Value = PerPlatformBuildConfig.Find(PlatformName);
 
-	return Value == nullptr ? EProjectPackagingBuildConfigurations::PPBC_Development : *Value;
+	// PPBC_MAX defines the default project setting case and should be handled accordingly.
+	return Value == nullptr ? EProjectPackagingBuildConfigurations::PPBC_MAX : *Value;
 }
 
 void UProjectPackagingSettings::SetBuildConfigurationForPlatform(FName PlatformName, EProjectPackagingBuildConfigurations Configuration)
@@ -227,17 +228,38 @@ void UProjectPackagingSettings::SetBuildConfigurationForPlatform(FName PlatformN
 	PerPlatformBuildConfig.Add(PlatformName, Configuration);
 }
 
-FName UProjectPackagingSettings::GetTargetPlatformForPlatform(FName PlatformName) const
+void UProjectPackagingSettings::SetBuildConfigurationForAllPlatforms(EProjectPackagingBuildConfigurations Configuration)
 {
-	const FName* Value = PerPlatformTargetPlatformName.Find(PlatformName);
+	for (auto& Config : PerPlatformBuildConfig)
+	{
+		Config.Value = Configuration;
+	}
+}
+
+FName UProjectPackagingSettings::GetTargetFlavorForPlatform(FName FlavorName) const
+{
+	const FName* Value = PerPlatformTargetFlavorName.Find(FlavorName);
+
+	// the flavor name is also the name of the vanilla info
+	return Value == nullptr ? FlavorName : *Value;
+}
+
+void UProjectPackagingSettings::SetTargetFlavorForPlatform(FName PlatformName, FName TargetFlavorName)
+{
+	PerPlatformTargetFlavorName.Add(PlatformName, TargetFlavorName);
+}
+
+FString UProjectPackagingSettings::GetBuildTargetForPlatform(FName PlatformName) const
+{
+	const FString* Value = PerPlatformBuildTarget.Find(PlatformName);
 
-	// the platform name is also the name of the vanilla info
-	return Value == nullptr ? PlatformName : *Value;
+	// empty string defines the default project setting case and should be handled accordingly.
+	return Value == nullptr ? "" : *Value;
 }
 
-void UProjectPackagingSettings::SetTargetPlatformForPlatform(FName PlatformName, FName TargetPlatformName)
+void UProjectPackagingSettings::SetBuildTargetForPlatform(FName PlatformName, FString BuildTargetName)
 {
-	PerPlatformTargetPlatformName.Add(PlatformName, TargetPlatformName);
+	PerPlatformBuildTarget.Add(PlatformName, BuildTargetName);
 }
 
 #undef LOCTEXT_NAMESPACE
\ No newline at end of file
diff --git a/Engine/Source/Editor/TurnkeySupport/Private/TurnkeySupportModule.cpp b/Engine/Source/Editor/TurnkeySupport/Private/TurnkeySupportModule.cpp
--- a/Engine/Source/Editor/TurnkeySupport/Private/TurnkeySupportModule.cpp
+++ b/Engine/Source/Editor/TurnkeySupport/Private/TurnkeySupportModule.cpp
@@ -259,7 +259,7 @@ public:
 		}
 		else
 		{
-			PlatformInfo = PlatformInfo::FindPlatformInfo(AllPlatformPackagingSettings->GetTargetPlatformForPlatform(IniPlatformName));
+			PlatformInfo = PlatformInfo::FindPlatformInfo(AllPlatformPackagingSettings->GetTargetFlavorForPlatform(IniPlatformName));
 		}
 		// this is unexpected to be able to happen, but it could if there was a bad value saved in the UProjectPackagingSettings - if this trips, we should handle errors
 		check(PlatformInfo != nullptr);
@@ -439,15 +439,39 @@ public:
 			}
 
 			// use AllPlatformPackagingSettings because these are user settings, and not checked in for all users to use
-			const UProjectPackagingSettings::FConfigurationInfo& ConfigurationInfo = UProjectPackagingSettings::ConfigurationInfo[(int)AllPlatformPackagingSettings->GetBuildConfigurationForPlatform(IniPlatformName)];
-			if (PlatformInfo->PlatformType == EBuildTargetType::Server)
+			FString BuildTarget = AllPlatformPackagingSettings->GetBuildTargetForPlatform(IniPlatformName);
+			if (BuildTarget.IsEmpty())
 			{
-				BuildCookRunParams += FString::Printf(TEXT(" -serverconfig=%s"), LexToString(ConfigurationInfo.Configuration));
+				BuildTarget = AllPlatformPackagingSettings->BuildTarget;
 			}
-			else
+			
+
+			EProjectPackagingBuildConfigurations BuildConfig = AllPlatformPackagingSettings->GetBuildConfigurationForPlatform(IniPlatformName);
+			// if PPBC_MAX is set, then the project default should be used instead of the per platform build config
+			if (BuildConfig == EProjectPackagingBuildConfigurations::PPBC_MAX)
+			{
+				BuildConfig = AllPlatformPackagingSettings->BuildConfiguration;
+			}
+
+			// when distribution is set, always package in shipping, which overrides the per platform build config
+			if (PackagingSettings->ForDistribution)
+			{
+				BuildConfig = EProjectPackagingBuildConfigurations::PPBC_Shipping;
+			}
+			
+			const UProjectPackagingSettings::FConfigurationInfo& ConfigurationInfo = UProjectPackagingSettings::ConfigurationInfo[(int)BuildConfig];
+			if (BuildTarget.IsEmpty())
 			{
 				BuildCookRunParams += FString::Printf(TEXT(" -clientconfig=%s"), LexToString(ConfigurationInfo.Configuration));
 			}
+			else if (PlatformInfo->PlatformType == EBuildTargetType::Server)
+			{
+				BuildCookRunParams += FString::Printf(TEXT(" -target=%s -serverconfig=%s"), *BuildTarget, LexToString(ConfigurationInfo.Configuration));
+			}
+			else
+			{
+				BuildCookRunParams += FString::Printf(TEXT(" -target=%s -clientconfig=%s"), *BuildTarget, LexToString(ConfigurationInfo.Configuration));
+			}
 		}
 // 		else if (Mode == EPrepareContentMode::PrepareForDebugging)
 // 		{
@@ -504,7 +528,7 @@ public:
 
 	static void ExecuteCustomBuild(FName IniPlatformName, FProjectBuildSettings Build)
 	{
-		const PlatformInfo::FTargetPlatformInfo* PlatformInfo = PlatformInfo::FindPlatformInfo(GetDefault<UProjectPackagingSettings>()->GetTargetPlatformForPlatform(IniPlatformName));
+		const PlatformInfo::FTargetPlatformInfo* PlatformInfo = PlatformInfo::FindPlatformInfo(GetDefault<UProjectPackagingSettings>()->GetTargetFlavorForPlatform(IniPlatformName));
 		const FString ProjectPath = GetProjectPathForTurnkey();
 
 		FString CommandLine;
@@ -534,26 +558,48 @@ public:
 		return true;
 	}
 
-	static bool PackageBuildConfigurationIsChecked(const PlatformInfo::FTargetPlatformInfo* Info, EProjectPackagingBuildConfigurations BuildConfiguration)
+	static bool DefaultPackageBuildConfigurationIsChecked(EProjectPackagingBuildConfigurations BuildConfiguration)
 	{
-		return GetDefault<UProjectPackagingSettings>()->GetBuildConfigurationForPlatform(Info->IniPlatformName) == BuildConfiguration;
+		return BuildConfiguration == EProjectPackagingBuildConfigurations::PPBC_MAX;
 	}
 
-	static void SetActiveTargetPlatform(const PlatformInfo::FTargetPlatformInfo* Info)
+	static bool PackageBuildConfigurationIsChecked(const PlatformInfo::FTargetPlatformInfo* Info, EProjectPackagingBuildConfigurations BuildConfiguration, EProjectPackagingBuildConfigurations DefaultBuildConfiguration)
+	{
+		return DefaultBuildConfiguration != EProjectPackagingBuildConfigurations::PPBC_MAX && GetDefault<UProjectPackagingSettings>()->GetBuildConfigurationForPlatform(Info->IniPlatformName) == BuildConfiguration;
+	}	
+	
+	static void SetActiveFlavor(const PlatformInfo::FTargetPlatformInfo* Info)
 	{
 		UProjectPackagingSettings* PackagingSettings = GetMutableDefault<UProjectPackagingSettings>();
-		PackagingSettings->SetTargetPlatformForPlatform(Info->IniPlatformName, Info->Name);
+		PackagingSettings->SetTargetFlavorForPlatform(Info->IniPlatformName, Info->Name);
 		PackagingSettings->SaveConfig();
 	}
 
-	static bool CanSetActiveTargetPlatform(const PlatformInfo::FTargetPlatformInfo* Info)
+	static bool CanSetActiveFlavor(const PlatformInfo::FTargetPlatformInfo* Info)
 	{
 		return true;
 	}
 
-	static bool SetActiveTargetPlatformIsChecked(const PlatformInfo::FTargetPlatformInfo* Info)
+	static bool SetActiveFlavorIsChecked(const PlatformInfo::FTargetPlatformInfo* Info)
+	{
+		return GetDefault<UProjectPackagingSettings>()->GetTargetFlavorForPlatform(Info->IniPlatformName) == Info->Name;
+	}
+
+	static void PackageBuildTarget(const PlatformInfo::FTargetPlatformInfo* Info, FString TargetName)
+	{
+		UProjectPackagingSettings* PackagingSettings = GetMutableDefault<UProjectPackagingSettings>();
+		PackagingSettings->SetBuildTargetForPlatform(Info->IniPlatformName, TargetName);
+		PackagingSettings->SaveConfig();
+	}
+	
+	static bool PackageBuildTargetIsChecked(const PlatformInfo::FTargetPlatformInfo* Info, FString TargetName, FString DefaultTargetName)
+	{
+		return !DefaultTargetName.IsEmpty() && GetDefault<UProjectPackagingSettings>()->GetBuildTargetForPlatform(Info->IniPlatformName) == TargetName;
+	}
+
+	static bool DefaultPackageBuildTargetIsChecked(FString TargetName)
 	{
-		return GetDefault<UProjectPackagingSettings>()->GetTargetPlatformForPlatform(Info->IniPlatformName) == Info->Name;
+		return TargetName.IsEmpty();
 	}
 
 	static void SetCookOnTheFly()
@@ -906,63 +952,157 @@ static void MakeTurnkeyPlatformMenu(FMenuBuilder& MenuBuilder, FName IniPlatform
 
 		MenuBuilder.EndSection();
 
+		UProjectPackagingSettings* AllPlatformPackagingSettings = GetMutableDefault<UProjectPackagingSettings>();
 
-		MenuBuilder.BeginSection("BuildConfig", LOCTEXT("TurnkeySection_BuildConfig", "Binary Configuration"));
-			EProjectType ProjectType = FTurnkeyEditorSupport::DoesProjectHaveCode() ? EProjectType::Code : EProjectType::Content;
-			TArray<EProjectPackagingBuildConfigurations> PackagingConfigurations = UProjectPackagingSettings::GetValidPackageConfigurations();
+		// Flavor Selection exists for Android to be able to package ASTC, DXT, ETC2, etc
+		{
+			// gather all valid flavors
+			const TArray<const PlatformInfo::FTargetPlatformInfo*> ValidFlavors = VanillaInfo->Flavors.FilterByPredicate([](const PlatformInfo::FTargetPlatformInfo* Target)
+				{
+					// Editor isn't a valid platform type that users can target
+					// The Build Target will choose client or server, so no need to show them as well
+					return Target->PlatformType != EBuildTargetType::Editor && Target->PlatformType != EBuildTargetType::Client && Target->PlatformType != EBuildTargetType::Server;
+				});
 
-			for (EProjectPackagingBuildConfigurations PackagingConfiguration : PackagingConfigurations)
+			if (ValidFlavors.Num() > 1)
 			{
-				const UProjectPackagingSettings::FConfigurationInfo& ConfigurationInfo = UProjectPackagingSettings::ConfigurationInfo[(int)PackagingConfiguration];
-				if (FInstalledPlatformInfo::Get().IsValid(TOptional<EBuildTargetType>(), TOptional<FString>(), ConfigurationInfo.Configuration, ProjectType, EInstalledPlatformState::Downloaded))
+				// Set the first flavor as the default if it hasn't been set
+				FName CurrentFlavor = AllPlatformPackagingSettings->GetTargetFlavorForPlatform(IniPlatformName);
+				if (CurrentFlavor == VanillaInfo->Name)
+				{
+					FTurnkeySupportCallbacks::SetActiveFlavor(ValidFlavors[0]);
+				}
+				
+				MenuBuilder.BeginSection("FlavorSelection", LOCTEXT("TurnkeySection_FlavorSelection", "Flavor Selection"));
+				
+				for (const PlatformInfo::FTargetPlatformInfo* Info : ValidFlavors)
 				{
 					MenuBuilder.AddMenuEntry(
-						ConfigurationInfo.Name,
-						ConfigurationInfo.ToolTip,
+						Info->DisplayName,
+						FText(),
 						FSlateIcon(),
 						FUIAction(
-							FExecuteAction::CreateStatic(&FTurnkeySupportCallbacks::PackageBuildConfiguration, VanillaInfo, PackagingConfiguration),
-							FCanExecuteAction::CreateStatic(&FTurnkeySupportCallbacks::CanPackageBuildConfiguration, VanillaInfo, PackagingConfiguration),
-							FIsActionChecked::CreateStatic(&FTurnkeySupportCallbacks::PackageBuildConfigurationIsChecked, VanillaInfo, PackagingConfiguration)
+							FExecuteAction::CreateStatic(&FTurnkeySupportCallbacks::SetActiveFlavor, Info),
+							FCanExecuteAction::CreateStatic(&FTurnkeySupportCallbacks::CanSetActiveFlavor, Info),
+							FIsActionChecked::CreateStatic(&FTurnkeySupportCallbacks::SetActiveFlavorIsChecked, Info)
 						),
 						NAME_None,
 						EUserInterfaceActionType::RadioButton
 					);
 				}
+				MenuBuilder.EndSection();
 			}
+		}
+
+		MenuBuilder.BeginSection("BuildConfig", LOCTEXT("TurnkeySection_BuildConfig", "Binary Configuration"));
+
+		// Get the enum metadata to display the correct menu name
+		EProjectPackagingBuildConfigurations DefaultBuildConfiguration = AllPlatformPackagingSettings->GetBuildConfigurationForPlatform(IniPlatformName);
+
+		UEnum* Enum = StaticEnum<EProjectPackagingBuildConfigurations>();
+		check(Enum);
+		FString Metadata = Enum->GetMetaData(TEXT("DisplayName"), (int32)AllPlatformPackagingSettings->BuildConfiguration);
+
+		MenuBuilder.AddMenuEntry(
+			FText::Format(LOCTEXT("DefaultConfiguration",  "Use Project Setting ({0})"), FText::FromString(Metadata)),
+			FText::Format(LOCTEXT("DefaultConfigurationTooltip", "Package the game in {0} configuration"), FText::FromString(Metadata)),
+			FSlateIcon(),
+			FUIAction(
+				FExecuteAction::CreateStatic(&FTurnkeySupportCallbacks::PackageBuildConfiguration, VanillaInfo, EProjectPackagingBuildConfigurations::PPBC_MAX),
+				FCanExecuteAction::CreateStatic(&FTurnkeySupportCallbacks::CanPackageBuildConfiguration, VanillaInfo, DefaultBuildConfiguration),
+				FIsActionChecked::CreateStatic(&FTurnkeySupportCallbacks::DefaultPackageBuildConfigurationIsChecked, DefaultBuildConfiguration)
+			),
+			NAME_None,
+			EUserInterfaceActionType::RadioButton
+		);
+
+		EProjectType ProjectType = FTurnkeyEditorSupport::DoesProjectHaveCode() ? EProjectType::Code : EProjectType::Content;
+		TArray<EProjectPackagingBuildConfigurations> PackagingConfigurations = UProjectPackagingSettings::GetValidPackageConfigurations();
+
+		for (EProjectPackagingBuildConfigurations PackagingConfiguration : PackagingConfigurations)
+		{
+			const UProjectPackagingSettings::FConfigurationInfo& ConfigurationInfo = UProjectPackagingSettings::ConfigurationInfo[(int)PackagingConfiguration];
+			if (FInstalledPlatformInfo::Get().IsValid(TOptional<EBuildTargetType>(), TOptional<FString>(), ConfigurationInfo.Configuration, ProjectType, EInstalledPlatformState::Downloaded))
+			{
+				MenuBuilder.AddMenuEntry(
+					ConfigurationInfo.Name,
+					ConfigurationInfo.ToolTip,
+					FSlateIcon(),
+					FUIAction(
+						FExecuteAction::CreateStatic(&FTurnkeySupportCallbacks::PackageBuildConfiguration, VanillaInfo, PackagingConfiguration),
+						FCanExecuteAction::CreateStatic(&FTurnkeySupportCallbacks::CanPackageBuildConfiguration, VanillaInfo, PackagingConfiguration),
+						FIsActionChecked::CreateStatic(&FTurnkeySupportCallbacks::PackageBuildConfigurationIsChecked, VanillaInfo, PackagingConfiguration, DefaultBuildConfiguration)
+					),
+					NAME_None,
+					EUserInterfaceActionType::RadioButton
+				);
+			}
+		}
 		MenuBuilder.EndSection();
 
-		// Binary builds can't compile for Client, Server or Editor so that leaves only the default option. No point creating a section for one option. 
-		if (!FApp::IsEngineInstalled())
+		IDesktopPlatform* DesktopPlatform = FDesktopPlatformModule::Get();
+		// Remove empty "Build Targets" menu item for content projects.
+		if (DesktopPlatform->GetTargetsForCurrentProject().Num() > 0)
 		{
-			MenuBuilder.BeginSection("TargetSelection", LOCTEXT("TurnkeySection_TargetSelection", "Target Selection"));
-			
-			// gather all platform infos
-			TArray<const PlatformInfo::FTargetPlatformInfo*> AllTargets = { VanillaInfo };
-			AllTargets.Append(VanillaInfo->Flavors);
+			TArray<FTargetInfo> Targets = DesktopPlatform->GetTargetsForCurrentProject();
+
+			Targets.Sort([](const FTargetInfo& A, const FTargetInfo& B) { return A.Name < B.Name; });
 			
-			for (const PlatformInfo::FTargetPlatformInfo* Info : AllTargets)
+			const TArray<FTargetInfo> ValidTargets = Targets.FilterByPredicate([](const FTargetInfo& Target)
+				{
+					return Target.Type == EBuildTargetType::Game || Target.Type == EBuildTargetType::Client || Target.Type == EBuildTargetType::Server;
+				});
+
+			if (ValidTargets.Num() > 1)
 			{
-				// Editor isn't a valid platform type that users can target
-				if (Info->PlatformType == EBuildTargetType::Editor)
+				// Set BuildTarget to default to Game if it hasn't been set
+				if(AllPlatformPackagingSettings->BuildTarget.IsEmpty())
 				{
-					continue;
+					const TArray<FTargetInfo> GameTarget = ValidTargets.FilterByPredicate([](const FTargetInfo& Target)
+						{
+							return Target.Type == EBuildTargetType::Game;
+						});
+					check(GameTarget.Num() > 0);
+
+					AllPlatformPackagingSettings->BuildTarget = GameTarget[0].Name;
+					AllPlatformPackagingSettings->SaveConfig();
 				}
-				
+
+				FText DefaultTarget = FText::FromString(AllPlatformPackagingSettings->GetBuildTargetForPlatform(IniPlatformName));
+
+				MenuBuilder.BeginSection("BuildTarget", LOCTEXT("TurnkeySection_BuildTarget", "Build Target"));
+
 				MenuBuilder.AddMenuEntry(
-					Info->DisplayName,
-					FText(),
+					FText::Format(LOCTEXT("DefaultPackageTarget",  "Use Project Setting ({0})"), FText::FromString(AllPlatformPackagingSettings->BuildTarget)),
+					FText::Format(LOCTEXT("DefaultPackageTargetTooltip", "Package the {0} target"), FText::FromString(AllPlatformPackagingSettings->BuildTarget)),
 					FSlateIcon(),
 					FUIAction(
-						FExecuteAction::CreateStatic(&FTurnkeySupportCallbacks::SetActiveTargetPlatform, Info),
-						FCanExecuteAction::CreateStatic(&FTurnkeySupportCallbacks::CanSetActiveTargetPlatform, Info),
-						FIsActionChecked::CreateStatic(&FTurnkeySupportCallbacks::SetActiveTargetPlatformIsChecked, Info)
-					),
+						FExecuteAction::CreateStatic(&FTurnkeySupportCallbacks::PackageBuildTarget, VanillaInfo, FString("")),
+						FCanExecuteAction(),
+						FIsActionChecked::CreateStatic(&FTurnkeySupportCallbacks::DefaultPackageBuildTargetIsChecked, DefaultTarget.ToString())
+						),
 					NAME_None,
 					EUserInterfaceActionType::RadioButton
 				);
+
+				for (const FTargetInfo& Target : ValidTargets)
+				{
+					MenuBuilder.AddMenuEntry(
+						FText::FromString(Target.Name),
+						FText::Format(LOCTEXT("PackageTargetName", "Package the '{0}' target."), FText::FromString(Target.Name)),
+						FSlateIcon(),
+						FUIAction(
+							FExecuteAction::CreateStatic(&FTurnkeySupportCallbacks::PackageBuildTarget, VanillaInfo, Target.Name),
+							FCanExecuteAction(),
+							FIsActionChecked::CreateStatic(&FTurnkeySupportCallbacks::PackageBuildTargetIsChecked, VanillaInfo, Target.Name, DefaultTarget.ToString())
+						),
+						NAME_None,
+						EUserInterfaceActionType::RadioButton
+					);
+				}
+				
+				MenuBuilder.EndSection();
 			}
-			MenuBuilder.EndSection();
 		}
 
 		MenuBuilder.BeginSection("AllDevices", LOCTEXT("TurnkeySection_AllDevices", "All Devices"));
