TurnkeySupportModule: Fix Packaging crash where UProjectPackagingSettings was getting GCed.

diff --git a/Engine/Source/Editor/TurnkeySupport/Private/TurnkeySupportModule.cpp b/Engine/Source/Editor/TurnkeySupport/Private/TurnkeySupportModule.cpp
--- a/Engine/Source/Editor/TurnkeySupport/Private/TurnkeySupportModule.cpp
+++ b/Engine/Source/Editor/TurnkeySupport/Private/TurnkeySupportModule.cpp
@@ -238,6 +238,8 @@ public:
 		if (PackagingSettings == nullptr)
 		{
 			PackagingSettings = NewObject<UProjectPackagingSettings>(GetTransientPackage());
+			// Prevent object from being GCed.
+			PackagingSettings->AddToRoot();
 			// make sure any changes to DefaultGame are updated in this class
 			PackagingSettings->LoadSettingsForPlatform(PlatformString);
 		}
