Fixes for Turnkey Launch On
- When a client build target is selected, launch on a corresponding client device
- Pass the UAT command line parameters for the current platform on launch on (flavor, build target)

diff --git a/Engine/Source/Editor/TurnkeySupport/Private/TurnkeyEditorSupport.cpp b/Engine/Source/Editor/TurnkeySupport/Private/TurnkeyEditorSupport.cpp
--- a/Engine/Source/Editor/TurnkeySupport/Private/TurnkeyEditorSupport.cpp
+++ b/Engine/Source/Editor/TurnkeySupport/Private/TurnkeyEditorSupport.cpp
@@ -22,6 +22,7 @@
 #include "Interfaces/IMainFrameModule.h"
 #include "ToolMenus.h"
 #include "FileHelpers.h"
+#include "Settings/ProjectPackagingSettings.h"
 #endif
 
 #define LOCTEXT_NAMESPACE "FTurnkeyEditorSupport"
@@ -87,7 +88,16 @@ void FTurnkeyEditorSupport::LaunchRunningMap(const FString& DeviceId, const FStr
 	FTargetDeviceId TargetDeviceId;
 	if (FTargetDeviceId::Parse(DeviceId, TargetDeviceId))
 	{
-		const PlatformInfo::FTargetPlatformInfo* const PlatformInfo = PlatformInfo::FindPlatformInfo(*TargetDeviceId.GetPlatformName());
+		const PlatformInfo::FTargetPlatformInfo* PlatformInfo = nullptr;
+		if (FApp::IsInstalled())
+		{
+			PlatformInfo = PlatformInfo::FindPlatformInfo(*TargetDeviceId.GetPlatformName());
+		}
+		else
+		{
+			PlatformInfo = PlatformInfo::FindPlatformInfo(GetDefault<UProjectPackagingSettings>()->GetTargetFlavorForPlatform(*TargetDeviceId.GetPlatformName()));
+		}
+					
 		FString UBTPlatformName = PlatformInfo->DataDrivenPlatformInfo->UBTPlatformString;
 		FString IniPlatformName = PlatformInfo->IniPlatformName.ToString();
 
@@ -120,7 +130,7 @@ void FTurnkeyEditorSupport::LaunchRunningMap(const FString& DeviceId, const FStr
 			// if we want to check device flash before we start cooking, kick it off now. we could delay this 
 			if (bUseTurnkey)
 			{
-				FString CommandLine = FString::Printf(TEXT("Turnkey -command=VerifySdk -UpdateIfNeeded -platform=%s -EditorIO -noturnkeyvariables -device=%s -utf8output -WaitForUATMutex"), *UBTPlatformName, *TargetDeviceId.GetDeviceName());
+				FString CommandLine = FString::Printf(TEXT("Turnkey -command=VerifySdk -UpdateIfNeeded -platform=%s -EditorIO -noturnkeyvariables -device=%s -utf8output -WaitForUATMutex %s"), *UBTPlatformName, *TargetDeviceId.GetDeviceName(), *PlatformInfo->UATCommandLine);
 				if (!ProjectPath.IsEmpty())
 				{
 					CommandLine = FString::Printf(TEXT(" -ScriptsForProject=\"%s\" %s -project=\"%s\""), *ProjectPath, *CommandLine, *ProjectPath);
diff --git a/Engine/Source/Editor/TurnkeySupport/Private/TurnkeySupportModule.cpp b/Engine/Source/Editor/TurnkeySupport/Private/TurnkeySupportModule.cpp
--- a/Engine/Source/Editor/TurnkeySupport/Private/TurnkeySupportModule.cpp
+++ b/Engine/Source/Editor/TurnkeySupport/Private/TurnkeySupportModule.cpp
@@ -1343,9 +1343,54 @@ static void GenerateDeviceProxyMenuParams(TSharedPtr<ITargetDeviceProxy> DeviceP
 
 		// ... create an action...
 	OutAction = FUIAction(
-		FExecuteAction::CreateLambda([DeviceProxy, ExternalOnClickDelegate]()
+		FExecuteAction::CreateLambda([DeviceProxy, PlatformName, ExternalOnClickDelegate]()
 			{
-				FString DeviceId = DeviceProxy->GetTargetDeviceId(NAME_None);
+				// only game and client devices are supported for launch on but devices only signal if they target client builds by their name e.g. "PS4Client"
+				// if the user has a EBuildTargetType::Client target selected we will launch on a client device, otherwise we fall back to the default game device
+
+				// We need to use flavors to launch on correctly on Android_ASTC / Android_ETC2, etc
+				const PlatformInfo::FTargetPlatformInfo* PlatformInfo = nullptr;
+				if (FApp::IsInstalled())
+				{
+					PlatformInfo = PlatformInfo::FindPlatformInfo(PlatformName);
+				}
+				else
+				{
+					PlatformInfo = PlatformInfo::FindPlatformInfo(GetDefault<UProjectPackagingSettings>()->GetTargetFlavorForPlatform(PlatformName));
+				}
+
+				FString VariantName = PlatformInfo->Name.ToString();
+
+				// find out if the user selected a client build target 
+				IDesktopPlatform* DesktopPlatform = FDesktopPlatformModule::Get();
+				if (DesktopPlatform->GetTargetsForCurrentProject().Num() > 0)
+				{
+					TArray<FTargetInfo> Targets = DesktopPlatform->GetTargetsForCurrentProject();
+
+					const TArray<FTargetInfo> ClientTargets = Targets.FilterByPredicate([](const FTargetInfo& Target)
+						{
+							return Target.Type == EBuildTargetType::Client;
+						});
+
+					// we just want to know if any client build is selected, the correct flavor was picked above
+					for (const auto& Target : ClientTargets)
+					{
+						if (GetDefault<UProjectPackagingSettings>()->GetBuildTargetForPlatform(PlatformName) == Target.Name)
+						{
+							VariantName += "Client";
+							break;
+						}
+					}
+				}
+
+				// If the variant is not found it will fall back to default (NAME_None)
+				FName DeviceVariantName = NAME_None;
+				if (DeviceProxy->HasVariant(*VariantName))
+				{
+					DeviceVariantName = *VariantName;
+				}
+
+				FString DeviceId = DeviceProxy->GetTargetDeviceId(DeviceVariantName);
 				HandleLaunchOnDeviceActionExecute(DeviceId, DeviceProxy->GetName(), true);
 				ExternalOnClickDelegate.ExecuteIfBound(DeviceId);
 			}
