Added Turnkey and RunUAT console commands to the editor.

#rnx

diff --git a/Engine/Source/Editor/TurnkeySupport/Private/TurnkeySupportModule.cpp b/Engine/Source/Editor/TurnkeySupport/Private/TurnkeySupportModule.cpp
--- a/Engine/Source/Editor/TurnkeySupport/Private/TurnkeySupportModule.cpp
+++ b/Engine/Source/Editor/TurnkeySupport/Private/TurnkeySupportModule.cpp
@@ -2086,6 +2086,29 @@ void FTurnkeySupportModule::ClearDeviceStatus(FName PlatformName)
 
 
 
+#if WITH_EDITOR
+bool FTurnkeySupportModule::Exec( UWorld* InWorld, const TCHAR* Cmd, FOutputDevice& Ar )
+{
+	if (FParse::Command( &Cmd, TEXT("Turnkey")) )
+	{
+		// run Turnkey via UAT. The Cmd is added at the end in case the user wants to run additional commands
+		const FString CommandLine = FString::Printf( TEXT("Turnkey %s %s %s"), *ITurnkeyIOModule::Get().GetUATParams(), *FTurnkeyEditorSupport::GetUATOptions(), Cmd );
+		FTurnkeyEditorSupport::RunUAT(CommandLine, FText::GetEmpty(), LOCTEXT("Turnkey_CustomTurnkeyName", "Executing Turnkey"), LOCTEXT("Turnkey_CustomTurnkeyShortName", "Turnkey"), FEditorStyle::GetBrush(TEXT("MainFrame.PackageProject")));
+		return true;
+	}
+	else if ( FParse::Command( &Cmd, TEXT("RunUAT")) )
+	{
+		// run UAT directly. The Cmd is added at the start on the assumption that it contains the command to run
+		const FString CommandLine = FString::Printf( TEXT("%s %s %s"), Cmd, *ITurnkeyIOModule::Get().GetUATParams(), *FTurnkeyEditorSupport::GetUATOptions() );
+		FTurnkeyEditorSupport::RunUAT(CommandLine, FText::GetEmpty(), LOCTEXT("Turnkey_CustomUATName", "Executing Custom UAT Task"), LOCTEXT("Turnkey_CustomUATShortName", "UAT"), FEditorStyle::GetBrush(TEXT("MainFrame.PackageProject")));
+		return true;
+	}
+
+	return false;
+}
+#endif
+
+
 void FTurnkeySupportModule::StartupModule( )
 {
 	
diff --git a/Engine/Source/Editor/TurnkeySupport/Private/TurnkeySupportModule.h b/Engine/Source/Editor/TurnkeySupport/Private/TurnkeySupportModule.h
--- a/Engine/Source/Editor/TurnkeySupport/Private/TurnkeySupportModule.h
+++ b/Engine/Source/Editor/TurnkeySupport/Private/TurnkeySupportModule.h
@@ -4,12 +4,17 @@
 
 #include "CoreMinimal.h"
 #include "Interfaces/ITurnkeySupportModule.h"
-
+#if WITH_EDITOR
+	#include "Misc/CoreMisc.h"
+#endif
 
 /**
  * Editor main frame module
  */
 class FTurnkeySupportModule	: public ITurnkeySupportModule
+#if WITH_EDITOR
+	, public FSelfRegisteringExec
+#endif
 {
 public:
 
@@ -48,4 +53,10 @@ private:
 
 	// menu helpers
 	TSharedRef<SWidget> MakeTurnkeyMenuWidget() const;
+
+#if WITH_EDITOR
+	// FSelfRegisteringExec interface
+	virtual bool Exec( class UWorld* InWorld, const TCHAR* Cmd, FOutputDevice& Ar ) override;
+#endif
+
 };
diff --git a/Engine/Source/Editor/UATHelper/UATHelperModule.cpp b/Engine/Source/Editor/UATHelper/UATHelperModule.cpp
--- a/Engine/Source/Editor/UATHelper/UATHelperModule.cpp
+++ b/Engine/Source/Editor/UATHelper/UATHelperModule.cpp
@@ -423,7 +423,8 @@ public:
 		FFormatNamedArguments Arguments;
 		Arguments.Add(TEXT("Platform"), PlatformDisplayName);
 		Arguments.Add(TEXT("TaskName"), TaskName);
-		FNotificationInfo Info( FText::Format( LOCTEXT("UatTaskInProgressNotification", "{TaskName} for {Platform}..."), Arguments) );
+		FText NotificationFormat = (PlatformDisplayName.IsEmpty()) ? LOCTEXT("UatTaskInProgressNotificationNoPlatform", "{TaskName}...") : LOCTEXT("UatTaskInProgressNotification", "{TaskName} for {Platform}...");
+		FNotificationInfo Info( FText::Format( NotificationFormat, Arguments) );
 		
 		Info.Image = TaskIcon;
 		Info.bFireAndForget = false;
