Add FApp::IsUnattended() as an exit condition to Turnkey VerifySDK command

#rnx

diff --git a/Engine/Source/Editor/TurnkeySupport/Private/TurnkeySupportModule.cpp b/Engine/Source/Editor/TurnkeySupport/Private/TurnkeySupportModule.cpp
--- a/Engine/Source/Editor/TurnkeySupport/Private/TurnkeySupportModule.cpp
+++ b/Engine/Source/Editor/TurnkeySupport/Private/TurnkeySupportModule.cpp
@@ -1816,7 +1816,7 @@ void FTurnkeySupportModule::UpdateSdkInfo()
 	}
 
 	// don't run UAT from commandlets (like the cooker) that are often launched from UAT and this will go poorly
-	if (IsRunningCommandlet())
+	if (IsRunningCommandlet() || FApp::IsUnattended())
 	{
 		return;
 	}
