Cache Turnkey veryfication results to reduce iteration times

diff --git a/Engine/Source/Editor/TurnkeySupport/Private/TurnkeyEditorSupport.cpp b/Engine/Source/Editor/TurnkeySupport/Private/TurnkeyEditorSupport.cpp
--- a/Engine/Source/Editor/TurnkeySupport/Private/TurnkeyEditorSupport.cpp
+++ b/Engine/Source/Editor/TurnkeySupport/Private/TurnkeyEditorSupport.cpp
@@ -24,6 +24,12 @@
 #include "ToolMenus.h"
 #include "FileHelpers.h"
 #include "Settings/ProjectPackagingSettings.h"
+#include "Containers/Set.h"
+
+namespace
+{
+TSet<FString> VerifiedPlatformsAndDevices;
+}
 #endif
 
 #define LOCTEXT_NAMESPACE "FTurnkeyEditorSupport"
@@ -131,7 +137,25 @@ void FTurnkeyEditorSupport::LaunchRunningMap(const FString& DeviceId, const FStr
 			// if we want to check device flash before we start cooking, kick it off now. we could delay this 
 			if (bUseTurnkey)
 			{
-				FString CommandLine = FString::Printf(TEXT("Turnkey -command=VerifySdk -UpdateIfNeeded -platform=%s -noturnkeyvariables -device=%s -utf8output -WaitForUATMutex %s %s"), *UBTPlatformName, *TargetDeviceId.GetDeviceName(), *PlatformInfo->UATCommandLine, *ITurnkeyIOModule::Get().GetUATParams());
+				const FString& RealDeviceName = TargetDeviceId.GetDeviceName();
+				const bool bSkipPlatformCheck = VerifiedPlatformsAndDevices.Contains(UBTPlatformName);
+				const bool bSkipDeviceCheck = VerifiedPlatformsAndDevices.Contains(RealDeviceName);
+
+				FString CommandLine;
+				if (bSkipPlatformCheck && bSkipDeviceCheck)
+				{
+					GUnrealEd->RequestPlaySession(SessionParams);
+					return;
+				}
+				else if (bSkipPlatformCheck && !bSkipDeviceCheck)
+				{
+					CommandLine = FString::Printf(TEXT("Turnkey -command=VerifySdk -UpdateIfNeeded -platform=%s -SkipPlatform -noturnkeyvariables -device=%s -utf8output -WaitForUATMutex %s %s"), *UBTPlatformName, *RealDeviceName, *PlatformInfo->UATCommandLine, *ITurnkeyIOModule::Get().GetUATParams());
+				}
+				else
+				{
+					CommandLine = FString::Printf(TEXT("Turnkey -command=VerifySdk -UpdateIfNeeded -platform=%s -noturnkeyvariables -device=%s -utf8output -WaitForUATMutex %s %s"), *UBTPlatformName, *RealDeviceName, *PlatformInfo->UATCommandLine, *ITurnkeyIOModule::Get().GetUATParams());
+				}
+				
 				if (!ProjectPath.IsEmpty())
 				{
 					CommandLine = FString::Printf(TEXT(" -ScriptsForProject=\"%s\" %s -project=\"%s\""), *ProjectPath, *CommandLine, *ProjectPath);
@@ -139,15 +163,17 @@ void FTurnkeyEditorSupport::LaunchRunningMap(const FString& DeviceId, const FStr
 				FText TaskName = LOCTEXT("VerifyingSDK", "Verifying SDK and Device");
 
 				IUATHelperModule::Get().CreateUatTask(CommandLine, FText::FromString(IniPlatformName), TaskName, TaskName, FEditorStyle::GetBrush(TEXT("MainFrame.PackageProject")),
-					[SessionParams](FString Result, double)
+					[SessionParams, RealDeviceName, UBTPlatformName](FString Result, double)
 					{
 						// unfortunate string comparison for success
 						bool bWasSuccessful = Result == TEXT("Completed");
-						AsyncTask(ENamedThreads::GameThread, [SessionParams, bWasSuccessful]()
+						AsyncTask(ENamedThreads::GameThread, [SessionParams, bWasSuccessful, RealDeviceName, UBTPlatformName]()
 							{
 								if (bWasSuccessful)
 								{
 									GUnrealEd->RequestPlaySession(SessionParams);
+									VerifiedPlatformsAndDevices.Add(RealDeviceName);
+									VerifiedPlatformsAndDevices.Add(UBTPlatformName);
 								}
 								else
 								{
diff --git a/Engine/Source/Programs/AutomationTool/Turnkey/Commands/VerifySdk.cs b/Engine/Source/Programs/AutomationTool/Turnkey/Commands/VerifySdk.cs
--- a/Engine/Source/Programs/AutomationTool/Turnkey/Commands/VerifySdk.cs
+++ b/Engine/Source/Programs/AutomationTool/Turnkey/Commands/VerifySdk.cs
@@ -23,7 +23,7 @@ namespace Turnkey.Commands
 			bool bForceSdkInstall = TurnkeyUtils.ParseParam("ForceSdkInstall", CommandOptions);
 			bool bForceDeviceInstall = TurnkeyUtils.ParseParam("ForceDeviceInstall", CommandOptions);
 			bool bUpdateIfNeeded = bForceSdkInstall || bForceDeviceInstall || TurnkeyUtils.ParseParam("UpdateIfNeeded", CommandOptions);
-			string DeviceString = TurnkeyUtils.ParseParamValue("Device", null, CommandOptions);
+			bool bSkipPlatformCheck = TurnkeyUtils.ParseParam("SkipPlatform", CommandOptions);
 
 			List<UnrealTargetPlatform> PlatformsToCheck;
 			List<DeviceInfo> DevicesToCheck;
@@ -36,7 +36,7 @@ namespace Turnkey.Commands
 				TurnkeyUtils.Log("Devices were requested, but none of them were found. Since -platforms was not specified, exiting command...");
 			}
 
-			if (PlatformsToCheck.Count == 0)
+			if (PlatformsToCheck.Count == 0 && !bSkipPlatformCheck)
 			{
 				TurnkeyUtils.Log("Platform(s) and/or device(s) needed for VerifySdk command. Check parameters or selections.");
 				return;
