Fix cilent/server config flags not being set properly in Turnkey

diff --git a/Engine/Source/Developer/DeveloperToolSettings/Classes/Settings/ProjectPackagingSettings.h b/Engine/Source/Developer/DeveloperToolSettings/Classes/Settings/ProjectPackagingSettings.h
--- a/Engine/Source/Developer/DeveloperToolSettings/Classes/Settings/ProjectPackagingSettings.h
+++ b/Engine/Source/Developer/DeveloperToolSettings/Classes/Settings/ProjectPackagingSettings.h
@@ -592,6 +592,8 @@ public:
 	FString GetBuildTargetForPlatform(FName PlatformName) const;
 	void SetBuildTargetForPlatform(FName PlatformName, FString BuildTargetName);
 
+	const FTargetInfo* GetBuildTargetInfoForPlatform(FName PlatformName) const;
+
 private:
 	/** Per platform build configuration */
 	UPROPERTY(config)
diff --git a/Engine/Source/Developer/DeveloperToolSettings/Private/ProjectPackagingSettings.cpp b/Engine/Source/Developer/DeveloperToolSettings/Private/ProjectPackagingSettings.cpp
--- a/Engine/Source/Developer/DeveloperToolSettings/Private/ProjectPackagingSettings.cpp
+++ b/Engine/Source/Developer/DeveloperToolSettings/Private/ProjectPackagingSettings.cpp
@@ -215,6 +215,28 @@ const FTargetInfo* UProjectPackagingSettings::GetBuildTargetInfo() const
 	return (DefaultGameTarget != nullptr) ? DefaultGameTarget : DefaultClientTarget;
 }
 
+const FTargetInfo* UProjectPackagingSettings::GetBuildTargetInfoForPlatform(FName PlatformName) const
+{
+	const FTargetInfo* DefaultGameTarget = nullptr;
+	const FTargetInfo* DefaultClientTarget = nullptr;
+	for (const FTargetInfo& Target : FDesktopPlatformModule::Get()->GetTargetsForCurrentProject())
+	{
+		if (Target.Name == GetBuildTargetForPlatform(PlatformName))
+		{
+			return &Target;
+		}
+		else if (Target.Type == EBuildTargetType::Game && (DefaultGameTarget == nullptr || Target.Name < DefaultGameTarget->Name))
+		{
+			DefaultGameTarget = &Target;
+		}
+		else if (Target.Type == EBuildTargetType::Client && (DefaultClientTarget == nullptr || Target.Name < DefaultClientTarget->Name))
+		{
+			DefaultClientTarget = &Target;
+		}
+	}
+	return (DefaultGameTarget != nullptr) ? DefaultGameTarget : DefaultClientTarget;
+}
+
 EProjectPackagingBuildConfigurations UProjectPackagingSettings::GetBuildConfigurationForPlatform(FName PlatformName) const
 {
 	const EProjectPackagingBuildConfigurations* Value = PerPlatformBuildConfig.Find(PlatformName);
diff --git a/Engine/Source/Editor/TurnkeySupport/Private/TurnkeySupportModule.cpp b/Engine/Source/Editor/TurnkeySupport/Private/TurnkeySupportModule.cpp
--- a/Engine/Source/Editor/TurnkeySupport/Private/TurnkeySupportModule.cpp
+++ b/Engine/Source/Editor/TurnkeySupport/Private/TurnkeySupportModule.cpp
@@ -441,14 +441,6 @@ public:
 				BuildCookRunParams += FString::Printf(TEXT(" -manifests -createchunkinstall -chunkinstalldirectory=\"%s\" -chunkinstallversion=%s"), *(PackagingSettings->HttpChunkInstallDataDirectory.Path), *(PackagingSettings->HttpChunkInstallDataVersion));
 			}
 
-			// use AllPlatformPackagingSettings because these are user settings, and not checked in for all users to use
-			FString BuildTarget = AllPlatformPackagingSettings->GetBuildTargetForPlatform(IniPlatformName);
-			if (BuildTarget.IsEmpty())
-			{
-				BuildTarget = AllPlatformPackagingSettings->BuildTarget;
-			}
-			
-
 			EProjectPackagingBuildConfigurations BuildConfig = AllPlatformPackagingSettings->GetBuildConfigurationForPlatform(IniPlatformName);
 			// if PPBC_MAX is set, then the project default should be used instead of the per platform build config
 			if (BuildConfig == EProjectPackagingBuildConfigurations::PPBC_MAX)
@@ -462,18 +454,20 @@ public:
 				BuildConfig = EProjectPackagingBuildConfigurations::PPBC_Shipping;
 			}
 			
+			const FTargetInfo* BuildTargetInfo = AllPlatformPackagingSettings->GetBuildTargetInfoForPlatform(IniPlatformName);
+
 			const UProjectPackagingSettings::FConfigurationInfo& ConfigurationInfo = UProjectPackagingSettings::ConfigurationInfo[(int)BuildConfig];
-			if (BuildTarget.IsEmpty())
+			if (BuildTargetInfo->Type == EBuildTargetType::Client)
 			{
-				BuildCookRunParams += FString::Printf(TEXT(" -clientconfig=%s"), LexToString(ConfigurationInfo.Configuration));
+				BuildCookRunParams += FString::Printf(TEXT(" -client -target=%s -clientconfig=%s"), *BuildTargetInfo->Name, LexToString(ConfigurationInfo.Configuration));
 			}
-			else if (PlatformInfo->PlatformType == EBuildTargetType::Server)
+			else if (BuildTargetInfo->Type == EBuildTargetType::Server)
 			{
-				BuildCookRunParams += FString::Printf(TEXT(" -target=%s -serverconfig=%s"), *BuildTarget, LexToString(ConfigurationInfo.Configuration));
+				BuildCookRunParams += FString::Printf(TEXT(" -server -noclient -target=%s -serverconfig=%s"), *BuildTargetInfo->Name, LexToString(ConfigurationInfo.Configuration));
 			}
 			else
 			{
-				BuildCookRunParams += FString::Printf(TEXT(" -target=%s -clientconfig=%s"), *BuildTarget, LexToString(ConfigurationInfo.Configuration));
+				BuildCookRunParams += FString::Printf(TEXT(" -clientconfig=%s"), LexToString(ConfigurationInfo.Configuration));
 			}
 		}
 // 		else if (Mode == EPrepareContentMode::PrepareForDebugging)
