crash fix if there are no custom build targets configured when packaging for a platform.

#rnx

diff --git a/Engine/Source/Editor/TurnkeySupport/Private/TurnkeySupportModule.cpp b/Engine/Source/Editor/TurnkeySupport/Private/TurnkeySupportModule.cpp
--- a/Engine/Source/Editor/TurnkeySupport/Private/TurnkeySupportModule.cpp
+++ b/Engine/Source/Editor/TurnkeySupport/Private/TurnkeySupportModule.cpp
@@ -457,11 +457,11 @@ public:
 			const FTargetInfo* BuildTargetInfo = AllPlatformPackagingSettings->GetBuildTargetInfoForPlatform(IniPlatformName);
 
 			const UProjectPackagingSettings::FConfigurationInfo& ConfigurationInfo = UProjectPackagingSettings::ConfigurationInfo[(int)BuildConfig];
-			if (BuildTargetInfo->Type == EBuildTargetType::Client)
+			if (BuildTargetInfo && BuildTargetInfo->Type == EBuildTargetType::Client)
 			{
 				BuildCookRunParams += FString::Printf(TEXT(" -client -target=%s -clientconfig=%s"), *BuildTargetInfo->Name, LexToString(ConfigurationInfo.Configuration));
 			}
-			else if (BuildTargetInfo->Type == EBuildTargetType::Server)
+			else if (BuildTargetInfo && BuildTargetInfo->Type == EBuildTargetType::Server)
 			{
 				BuildCookRunParams += FString::Printf(TEXT(" -server -noclient -target=%s -serverconfig=%s"), *BuildTargetInfo->Name, LexToString(ConfigurationInfo.Configuration));
 			}
