build targets are shown and selectable when packaging blueprint projects.

#rnx

diff --git a/Engine/Source/Developer/DeveloperToolSettings/Private/ProjectPackagingSettings.cpp b/Engine/Source/Developer/DeveloperToolSettings/Private/ProjectPackagingSettings.cpp
--- a/Engine/Source/Developer/DeveloperToolSettings/Private/ProjectPackagingSettings.cpp
+++ b/Engine/Source/Developer/DeveloperToolSettings/Private/ProjectPackagingSettings.cpp
@@ -217,9 +217,16 @@ const FTargetInfo* UProjectPackagingSettings::GetBuildTargetInfo() const
 
 const FTargetInfo* UProjectPackagingSettings::GetBuildTargetInfoForPlatform(FName PlatformName) const
 {
+	// Collect build targets. Content-only projects use Engine targets.
+	FProjectStatus ProjectStatus;
+	bool bHasCode = IProjectManager::Get().QueryStatusForCurrentProject(ProjectStatus) && ProjectStatus.bCodeBasedProject;
+	
+	IDesktopPlatform* DesktopPlatform = FDesktopPlatformModule::Get();
+	const TArray<FTargetInfo>& TargetsRef = bHasCode ? DesktopPlatform->GetTargetsForCurrentProject() : DesktopPlatform->GetTargetsForProject(FString());
+
 	const FTargetInfo* DefaultGameTarget = nullptr;
 	const FTargetInfo* DefaultClientTarget = nullptr;
-	for (const FTargetInfo& Target : FDesktopPlatformModule::Get()->GetTargetsForCurrentProject())
+	for (const FTargetInfo& Target : TargetsRef)
 	{
 		if (Target.Name == GetBuildTargetForPlatform(PlatformName))
 		{
diff --git a/Engine/Source/Editor/TurnkeySupport/Private/TurnkeySupportModule.cpp b/Engine/Source/Editor/TurnkeySupport/Private/TurnkeySupportModule.cpp
--- a/Engine/Source/Editor/TurnkeySupport/Private/TurnkeySupportModule.cpp
+++ b/Engine/Source/Editor/TurnkeySupport/Private/TurnkeySupportModule.cpp
@@ -1037,12 +1037,15 @@ static void MakeTurnkeyPlatformMenu(FMenuBuilder& MenuBuilder, FName IniPlatform
 		}
 		MenuBuilder.EndSection();
 
+		// Collect build targets. Content-only projects use Engine targets
+		FProjectStatus ProjectStatus;
+		bool bHasCode = IProjectManager::Get().QueryStatusForCurrentProject(ProjectStatus) && ProjectStatus.bCodeBasedProject;
+
 		IDesktopPlatform* DesktopPlatform = FDesktopPlatformModule::Get();
-		// Remove empty "Build Targets" menu item for content projects.
-		if (DesktopPlatform->GetTargetsForCurrentProject().Num() > 0)
-		{
-			TArray<FTargetInfo> Targets = DesktopPlatform->GetTargetsForCurrentProject();
+		TArray<FTargetInfo> Targets = bHasCode ? DesktopPlatform->GetTargetsForCurrentProject() : DesktopPlatform->GetTargetsForProject(FString());
 
+		if (Targets.Num() > 0)
+		{
 			Targets.Sort([](const FTargetInfo& A, const FTargetInfo& B) { return A.Name < B.Name; });
 			
 			const TArray<FTargetInfo> ValidTargets = Targets.FilterByPredicate([](const FTargetInfo& Target)
