Fix Turnkey passing the wrong platform when updating device on Windows

diff --git a/Engine/Source/Editor/TurnkeySupport/Private/TurnkeySupportModule.cpp b/Engine/Source/Editor/TurnkeySupport/Private/TurnkeySupportModule.cpp
--- a/Engine/Source/Editor/TurnkeySupport/Private/TurnkeySupportModule.cpp
+++ b/Engine/Source/Editor/TurnkeySupport/Private/TurnkeySupportModule.cpp
@@ -703,7 +703,7 @@ static void ShowInstallationHelp(FName IniPlatformName)
 	FTurnkeyEditorSupport::ShowInstallationHelp(IniPlatformName, FDataDrivenPlatformInfoRegistry::GetPlatformInfo(IniPlatformName).SDKTutorial);
 }
 
-static void TurnkeyInstallSdk(FString PlatformName, bool bPreferFull, bool bForceInstall, FString DeviceId)
+static void TurnkeyInstallSdk(FString IniPlatformName, bool bPreferFull, bool bForceInstall, FString DeviceId)
 {
 	FString OptionalOptions;
 	if (bPreferFull)
@@ -725,13 +725,13 @@ static void TurnkeyInstallSdk(FString PlatformName, bool bPreferFull, bool bForc
 	{
 		CommandLine.Appendf(TEXT("-ScriptsForProject=\"%s\" "), *ProjectPath);
 	}
-	CommandLine.Appendf(TEXT("Turnkey -command=VerifySdk -UpdateIfNeeded -platform=%s %s %s -noturnkeyvariables -utf8output -WaitForUATMutex"), *PlatformName, *OptionalOptions, *ITurnkeyIOModule::Get().GetUATParams());
+	CommandLine.Appendf(TEXT("Turnkey -command=VerifySdk -UpdateIfNeeded -platform=%s %s %s -noturnkeyvariables -utf8output -WaitForUATMutex"), *IniPlatformName, *OptionalOptions, *ITurnkeyIOModule::Get().GetUATParams());
 
 	FText TaskName = LOCTEXT("InstallingSdk", "Installing Sdk");
-	FTurnkeyEditorSupport::RunUAT(CommandLine, FText::FromString(PlatformName), TaskName, TaskName, FEditorStyle::GetBrush(TEXT("MainFrame.PackageProject")),
-		[PlatformName](FString, double)
+	FTurnkeyEditorSupport::RunUAT(CommandLine, FText::FromString(IniPlatformName), TaskName, TaskName, FEditorStyle::GetBrush(TEXT("MainFrame.PackageProject")),
+		[IniPlatformName](FString, double)
 	{
-		AsyncTask(ENamedThreads::GameThread, [PlatformName]()
+		AsyncTask(ENamedThreads::GameThread, [IniPlatformName]()
 		{
 
 			// read in env var changes
@@ -767,7 +767,7 @@ static void TurnkeyInstallSdk(FString PlatformName, bool bPreferFull, bool bForc
 
 			// update the Sdk status
 //			FDataDrivenPlatformInfoRegistry::UpdateSdkStatus();
-			GetTargetPlatformManager()->UpdateAfterSDKInstall(*PlatformName);
+			GetTargetPlatformManager()->UpdateAfterSDKInstall(*IniPlatformName);
 #if WITH_ENGINE
 			RenderUtilsInit();
 #endif
@@ -1119,7 +1119,7 @@ static void MakeTurnkeyPlatformMenu(FMenuBuilder& MenuBuilder, FName IniPlatform
 			MenuBuilder.AddSubMenu(
 				MakeSdkStatusAttribute(IniPlatformName, Proxy),
 				FText(),
-				FNewMenuDelegate::CreateLambda([UBTPlatformString, IniPlatformName, DeviceName, DeviceId](FMenuBuilder& SubMenuBuilder)
+				FNewMenuDelegate::CreateLambda([IniPlatformName, DeviceName, DeviceId](FMenuBuilder& SubMenuBuilder)
 				{
 					FTurnkeySdkInfo SdkInfo = ITurnkeySupportModule::Get().GetSdkInfoForDeviceId(DeviceId);
 
@@ -1136,7 +1136,7 @@ static void MakeTurnkeyPlatformMenu(FMenuBuilder& MenuBuilder, FName IniPlatform
 							LOCTEXT("Turnkey_ForceRepairDevice", "Repair Device"),
 							LOCTEXT("TurnkeyTooltip_ForceRepairDevice", "Force repairing anything on the device needed (update firmware, etc). Will perform all steps possible, even if not needed."),
 							FSlateIcon(),
-							FExecuteAction::CreateStatic(TurnkeyInstallSdk, UBTPlatformString, true, false, DeviceId)
+							FExecuteAction::CreateStatic(TurnkeyInstallSdk, IniPlatformName.ToString(), true, false, DeviceId)
 						);
 					}
 					else
@@ -1145,7 +1145,7 @@ static void MakeTurnkeyPlatformMenu(FMenuBuilder& MenuBuilder, FName IniPlatform
 							LOCTEXT("Turnkey_RepairDevice", "Update Device"),
 							LOCTEXT("TurnkeyTooltip_RepairDevice", "Perform any fixup that may be needed on this device. If up to date already, nothing will be done."),
 							FSlateIcon(),
-							FExecuteAction::CreateStatic(TurnkeyInstallSdk, UBTPlatformString, false, false, DeviceId)
+							FExecuteAction::CreateStatic(TurnkeyInstallSdk, IniPlatformName.ToString(), false, false, DeviceId)
 						);
 					}
 
@@ -1197,7 +1197,7 @@ static void MakeTurnkeyPlatformMenu(FMenuBuilder& MenuBuilder, FName IniPlatform
 				LOCTEXT("Turnkey_UpdateSdkMinimal", "Update Sdk"),
 				LOCTEXT("TurnkeyTooltip_InstallSdkMinimal", "Attempt to update an Sdk, as hosted by your studio. Will attempt to install a minimal Sdk (useful for building/running only)"),
 				FSlateIcon(),
-				FExecuteAction::CreateStatic(TurnkeyInstallSdk, UBTPlatformString, false, false, NoDevice)
+				FExecuteAction::CreateStatic(TurnkeyInstallSdk, IniPlatformName.ToString(), false, false, NoDevice)
 			);
 
 			if (SdkInfo.bCanInstallFullSdk && SdkInfo.bCanInstallAutoSdk)
@@ -1206,7 +1206,7 @@ static void MakeTurnkeyPlatformMenu(FMenuBuilder& MenuBuilder, FName IniPlatform
 					LOCTEXT("Turnkey_UpdateSdkFull", "Update Sdk (Full Platform Installer)"),
 					LOCTEXT("TurnkeyTooltip_UpdateSdkFull", "Attempt to update an Sdk, as hosted by your studio. Will attempt to install a full Sdk (useful profiling or other use cases)"),
 					FSlateIcon(),
-					FExecuteAction::CreateStatic(TurnkeyInstallSdk, UBTPlatformString, true, false, NoDevice)
+					FExecuteAction::CreateStatic(TurnkeyInstallSdk, IniPlatformName.ToString(), true, false, NoDevice)
 				);
 			}
 		}
@@ -1216,7 +1216,7 @@ static void MakeTurnkeyPlatformMenu(FMenuBuilder& MenuBuilder, FName IniPlatform
 				LOCTEXT("Turnkey_ForceSdkMinimal", "Force Reinstall Sdk"),
 				LOCTEXT("TurnkeyTooltip_ForceSdkMinimal", "Attempt to force re-install an Sdk, as hosted by your studio. Will attempt to install a minimal Sdk (useful for building/running only)"),
 				FSlateIcon(),
-				FExecuteAction::CreateStatic(TurnkeyInstallSdk, UBTPlatformString, false, true, NoDevice)
+				FExecuteAction::CreateStatic(TurnkeyInstallSdk, IniPlatformName.ToString(), false, true, NoDevice)
 			);
 
 			if (SdkInfo.bCanInstallFullSdk && SdkInfo.bCanInstallAutoSdk)
@@ -1225,7 +1225,7 @@ static void MakeTurnkeyPlatformMenu(FMenuBuilder& MenuBuilder, FName IniPlatform
 					LOCTEXT("Turnkey_ForceSdkFull", "Force Reinstall (Full Platform Installer)"),
 					LOCTEXT("TurnkeyTooltip_ForceSdkForce", "Attempt to force re-install an Sdk, as hosted by your studio. Will attempt to install a full Sdk (useful profiling or other use cases)"),
 					FSlateIcon(),
-					FExecuteAction::CreateStatic(TurnkeyInstallSdk, UBTPlatformString, true, true, NoDevice)
+					FExecuteAction::CreateStatic(TurnkeyInstallSdk, IniPlatformName.ToString(), true, true, NoDevice)
 				);
 			}
 		}
@@ -1235,14 +1235,14 @@ static void MakeTurnkeyPlatformMenu(FMenuBuilder& MenuBuilder, FName IniPlatform
 				LOCTEXT("Turnkey_InstallSdkMinimal", "Install Sdk"),
 				LOCTEXT("TurnkeyTooltip_InstallSdkMinimal", "Attempt to install an Sdk, as hosted by your studio. Will attempt to install a minimal Sdk (useful for building/running only)"),
 				FSlateIcon(),
-				FExecuteAction::CreateStatic(TurnkeyInstallSdk, UBTPlatformString, false, false, NoDevice)
+				FExecuteAction::CreateStatic(TurnkeyInstallSdk, IniPlatformName.ToString(), false, false, NoDevice)
 			);
 
 			MenuBuilder.AddMenuEntry(
 				LOCTEXT("Turnkey_InstallSdkFull", "Install Sdk (Full Platform Installer)"),
 				LOCTEXT("TurnkeyTooltip_InstallSdkFull", "Attempt to install an Sdk, as hosted by your studio. Will attempt to install a full Sdk (useful profiling or other use cases)"),
 				FSlateIcon(),
-				FExecuteAction::CreateStatic(TurnkeyInstallSdk, UBTPlatformString, true, false, NoDevice)
+				FExecuteAction::CreateStatic(TurnkeyInstallSdk, IniPlatformName.ToString(), true, false, NoDevice)
 			);
 		}
 	}
