Adding the option to flag a platform as being hidden from the UI at runtime.

#rnx

diff --git a/Engine/Source/Editor/TurnkeySupport/Private/TurnkeySupportModule.cpp b/Engine/Source/Editor/TurnkeySupport/Private/TurnkeySupportModule.cpp
--- a/Engine/Source/Editor/TurnkeySupport/Private/TurnkeySupportModule.cpp
+++ b/Engine/Source/Editor/TurnkeySupport/Private/TurnkeySupportModule.cpp
@@ -1433,6 +1433,10 @@ void FTurnkeySupportModule::MakeQuickLaunchItems(class UToolMenu* Menu, FOnQuick
 
 			FName PlatformName = Pair.Key;
 			const FDataDrivenPlatformInfo& Info = Pair.Value;
+			if (FDataDrivenPlatformInfoRegistry::IsPlatformHiddenFromUI(PlatformName))
+			{
+				continue;
+			}
 
 			// look for devices for all platforms, even if the platform isn't installed - Turnkey can install Sdk after selecting LaunchOn
 			TArray<TSharedPtr<ITargetDeviceProxy>> DeviceProxies;
@@ -1575,6 +1579,10 @@ TSharedRef<SWidget> FTurnkeySupportModule::MakeTurnkeyMenuWidget() const
 
 				FName PlatformName = Pair.Key;
 				const FDataDrivenPlatformInfo& Info = Pair.Value;
+				if (FDataDrivenPlatformInfoRegistry::IsPlatformHiddenFromUI(PlatformName))
+				{
+					continue;
+				}
 
 				if (!FDataDrivenPlatformInfoRegistry::HasCompiledSupportForPlatform(PlatformName, FDataDrivenPlatformInfoRegistry::EPlatformNameType::Ini))
 				{	
diff --git a/Engine/Source/Runtime/Core/Private/Misc/DataDrivenPlatformInfoRegistry.cpp b/Engine/Source/Runtime/Core/Private/Misc/DataDrivenPlatformInfoRegistry.cpp
--- a/Engine/Source/Runtime/Core/Private/Misc/DataDrivenPlatformInfoRegistry.cpp
+++ b/Engine/Source/Runtime/Core/Private/Misc/DataDrivenPlatformInfoRegistry.cpp
@@ -7,6 +7,7 @@
 #include "Misc/ConfigCacheIni.h"
 #include "Modules/ModuleManager.h"
 #include "Misc/CommandLine.h"
+#include "Containers/Set.h"
 
 
 namespace 
@@ -17,6 +18,7 @@ namespace
 	TArray<const FDataDrivenPlatformInfo*> SortedPlatformInfos;
 #if DDPI_HAS_EXTENDED_PLATFORMINFO_DATA
 	TArray<struct FPreviewPlatformMenuItem> PreviewPlatformMenuItems;
+	TSet<FName> PlatformsHiddenFromUI;
 #endif
 }
 
@@ -587,4 +589,15 @@ const TArray<struct FPreviewPlatformMenuItem>& FDataDrivenPlatformInfoRegistry::
 	return PreviewPlatformMenuItems;
 }
 
+bool FDataDrivenPlatformInfoRegistry::IsPlatformHiddenFromUI(FName PlatformName)
+{
+	return PlatformsHiddenFromUI.Contains(PlatformName);
+}
+
+void FDataDrivenPlatformInfoRegistry::SetPlatformHiddenFromUI(FName PlatformName)
+{
+	PlatformsHiddenFromUI.Add(PlatformName);
+}
+
+
 #endif
diff --git a/Engine/Source/Runtime/Core/Public/Misc/DataDrivenPlatformInfoRegistry.h b/Engine/Source/Runtime/Core/Public/Misc/DataDrivenPlatformInfoRegistry.h
--- a/Engine/Source/Runtime/Core/Public/Misc/DataDrivenPlatformInfoRegistry.h
+++ b/Engine/Source/Runtime/Core/Public/Misc/DataDrivenPlatformInfoRegistry.h
@@ -308,6 +308,10 @@ struct CORE_API FDataDrivenPlatformInfoRegistry
 	*/
 	static const TArray<struct FPreviewPlatformMenuItem>& GetAllPreviewPlatformMenuItems();
 
+	/** Option to hide a platform from the user interface at runtime */
+	static bool IsPlatformHiddenFromUI(FName PlatformName);
+	static void SetPlatformHiddenFromUI(FName PlatformName);
+
 private:
 	/**
 	 * Get a modifiable DDPI object, for Turnkey to update it's info
