Add a refresh platform status button to Turnkey

diff --git a/Engine/Source/Editor/TurnkeySupport/Private/TurnkeySupportModule.cpp b/Engine/Source/Editor/TurnkeySupport/Private/TurnkeySupportModule.cpp
--- a/Engine/Source/Editor/TurnkeySupport/Private/TurnkeySupportModule.cpp
+++ b/Engine/Source/Editor/TurnkeySupport/Private/TurnkeySupportModule.cpp
@@ -667,6 +667,16 @@ public:
 	{
 		return GetDefault<UCookerSettings>()->bCookOnTheFlyForLaunchOn;
 	}
+
+	static void SetRefreshPlatformStatus()
+	{
+		ITurnkeySupportModule::Get().UpdateSdkInfo();
+	}
+
+	static bool CanRefreshPlatformStatus()
+	{
+		return true;
+	}
 };
 
 class FTurnkeySupportCommands : public TCommands<FTurnkeySupportCommands>
@@ -1651,6 +1661,16 @@ TSharedRef<SWidget> FTurnkeySupportModule::MakeTurnkeyMenuWidget() const
 				);
 			}
 
+			PlatformsSection.AddMenuEntry(
+				NAME_None,
+				LOCTEXT("RefreshPlatformStatus", "Refresh platform status"),
+				LOCTEXT("RefreshPlatformStatusDescription", "Update all platforms and device information"),
+				FSlateIcon(FAppStyle::GetAppStyleSetName(), "SourceControl.Actions.Refresh"),
+				FUIAction(
+				FExecuteAction::CreateStatic(&FTurnkeySupportCallbacks::SetRefreshPlatformStatus),
+				FCanExecuteAction::CreateStatic(&FTurnkeySupportCallbacks::CanRefreshPlatformStatus))
+			);
+
 			if (UnsupportedPlatforms.Num() != 0)
 			{
 				PlatformsSection.AddSeparator(NAME_None);
