Separated turnkey device and sdk status so they can be checked independently & fixed up device update ui.

#rnx

diff --git a/Engine/Source/Editor/TurnkeySupport/Private/TurnkeySupportModule.cpp b/Engine/Source/Editor/TurnkeySupport/Private/TurnkeySupportModule.cpp
--- a/Engine/Source/Editor/TurnkeySupport/Private/TurnkeySupportModule.cpp
+++ b/Engine/Source/Editor/TurnkeySupport/Private/TurnkeySupportModule.cpp
@@ -829,8 +829,9 @@ static FSlateIcon MakePlatformSdkIconAttribute(FName IniPlatformName, TSharedPtr
 		{
 			// get the status, or Unknown if it's not there
 			ETurnkeyPlatformSdkStatus Status = DeviceId.Len() ? ITurnkeySupportModule::Get().GetSdkInfoForDeviceId(DeviceId).Status : ITurnkeySupportModule::Get().GetSdkInfo(IniPlatformName, false).Status;
+			bool bDeviceWarning = (DeviceId.Len() && ITurnkeySupportModule::Get().GetSdkInfoForDeviceId(DeviceId).DeviceStatus != ETurnkeyDeviceStatus::SoftwareValid);
 
-			if (Status == ETurnkeyPlatformSdkStatus::OutOfDate || Status == ETurnkeyPlatformSdkStatus::NoSdk || Status == ETurnkeyPlatformSdkStatus::FlashOutOfDate)
+			if (Status == ETurnkeyPlatformSdkStatus::OutOfDate || Status == ETurnkeyPlatformSdkStatus::NoSdk || bDeviceWarning)
 			{
 				return FSlateIcon(FEditorStyle::GetStyleSetName(), TEXT("Icons.Warning"));
 			}
@@ -1152,13 +1153,13 @@ static void MakeTurnkeyPlatformMenu(FMenuBuilder& MenuBuilder, FName IniPlatform
 						FText::GetEmpty()
 					);
 
-					if (SdkInfo.Status == ETurnkeyPlatformSdkStatus::FlashValid)
+					if (SdkInfo.DeviceStatus == ETurnkeyDeviceStatus::SoftwareValid)
 					{
 						SubMenuBuilder.AddMenuEntry(
-							LOCTEXT("Turnkey_ForceRepairDevice", "Repair Device"),
+							LOCTEXT("Turnkey_ForceRepairDevice", "Force Update Device"),
 							LOCTEXT("TurnkeyTooltip_ForceRepairDevice", "Force repairing anything on the device needed (update firmware, etc). Will perform all steps possible, even if not needed."),
 							FSlateIcon(),
-							FExecuteAction::CreateStatic(TurnkeyInstallSdk, IniPlatformName.ToString(), true, false, DeviceId)
+							FExecuteAction::CreateStatic(TurnkeyInstallSdk, IniPlatformName.ToString(), false, true, DeviceId)
 						);
 					}
 					else
@@ -1868,6 +1869,20 @@ bool GetSdkInfoFromTurnkey(FString Line, FName& PlatformName, FString& DeviceId,
 	SdkInfo.bCanInstallFullSdk = FlagsString.Contains(TEXT("Support_FullSdk"));
 	SdkInfo.bCanInstallAutoSdk = FlagsString.Contains(TEXT("Support_AutoSdk"));
 
+	SdkInfo.DeviceStatus = ETurnkeyDeviceStatus::Unknown;
+	if (FlagsString.Contains(TEXT("Device_InvalidPrerequisites")))
+	{
+		SdkInfo.DeviceStatus = ETurnkeyDeviceStatus::InvalidPrerequisites;
+	}
+	else if (FlagsString.Contains(TEXT("Device_InstallSoftwareValid")))
+	{
+		SdkInfo.DeviceStatus = ETurnkeyDeviceStatus::SoftwareValid;
+	}
+	else if (FlagsString.Contains(TEXT("Device_InstallSoftwareInvalid")))
+	{
+		SdkInfo.DeviceStatus = ETurnkeyDeviceStatus::SoftwareInvalid;
+	}
+
 	return true;
 }
 
diff --git a/Engine/Source/Editor/TurnkeySupport/Public/Interfaces/ITurnkeySupportModule.h b/Engine/Source/Editor/TurnkeySupport/Public/Interfaces/ITurnkeySupportModule.h
--- a/Engine/Source/Editor/TurnkeySupport/Public/Interfaces/ITurnkeySupportModule.h
+++ b/Engine/Source/Editor/TurnkeySupport/Public/Interfaces/ITurnkeySupportModule.h
@@ -17,14 +17,21 @@ enum class ETurnkeyPlatformSdkStatus : uint8
 	OutOfDate,
 	NoSdk,
 	Error,
-	FlashValid,
-	FlashOutOfDate,
 	// @todo turnkey: add AutoSdkValid and ManualSdkValid, with Valid a Combination of both
 };
 
+enum class ETurnkeyDeviceStatus : uint8
+{
+	Unknown,
+	InvalidPrerequisites,
+	SoftwareValid,
+	SoftwareInvalid,
+};
+
 struct FTurnkeySdkInfo
 {
 	ETurnkeyPlatformSdkStatus Status = ETurnkeyPlatformSdkStatus::Unknown;
+	ETurnkeyDeviceStatus DeviceStatus = ETurnkeyDeviceStatus::Unknown;
 	FText SdkErrorInformation;
 	FString InstalledVersion;
 	FString AutoSDKVersion; // only valid for platform, not device
