Fix compile issue for non-editor builds in TurnkeySupport module

diff --git a/Engine/Source/Editor/TurnkeySupport/Private/TurnkeySupportModule.cpp b/Engine/Source/Editor/TurnkeySupport/Private/TurnkeySupportModule.cpp
--- a/Engine/Source/Editor/TurnkeySupport/Private/TurnkeySupportModule.cpp
+++ b/Engine/Source/Editor/TurnkeySupport/Private/TurnkeySupportModule.cpp
@@ -1022,14 +1022,11 @@ static void MakeTurnkeyPlatformMenu(FMenuBuilder& MenuBuilder, FName IniPlatform
 		}
 
 		MenuBuilder.BeginSection("BuildConfig", LOCTEXT("TurnkeySection_BuildConfig", "Binary Configuration"));
-
-		UEnum* Enum = StaticEnum<EProjectPackagingBuildConfigurations>();
-		check(Enum);
-		FString Metadata = Enum->GetMetaData(TEXT("DisplayName"), (int32)AllPlatformPackagingSettings->BuildConfiguration);
-
+		
+		const UProjectPackagingSettings::FConfigurationInfo& ConfigInfo = UProjectPackagingSettings::ConfigurationInfo[static_cast<int32>(AllPlatformPackagingSettings->BuildConfiguration)];
 		MenuBuilder.AddMenuEntry(
-			FText::Format(LOCTEXT("DefaultConfiguration",  "Use Project Setting ({0})"), FText::FromString(Metadata)),
-			FText::Format(LOCTEXT("DefaultConfigurationTooltip", "Package the game in {0} configuration"), FText::FromString(Metadata)),
+			FText::Format(LOCTEXT("DefaultConfiguration",  "Use Project Setting ({0})"), ConfigInfo.Name),
+			ConfigInfo.ToolTip,
 			FSlateIcon(),
 			FUIAction(
 				FExecuteAction::CreateStatic(&FTurnkeySupportCallbacks::SetPackageBuildConfiguration, VanillaInfo, EProjectPackagingBuildConfigurations::PPBC_MAX),
