- Changed TurnkeyPlatform menu to use UToolMenu instead of legacy FMenuBuilder
- Added tooltips for how to get the SdkInfo from commandline

diff --git a/Engine/Source/Editor/TurnkeySupport/Private/TurnkeySupportModule.cpp b/Engine/Source/Editor/TurnkeySupport/Private/TurnkeySupportModule.cpp
--- a/Engine/Source/Editor/TurnkeySupport/Private/TurnkeySupportModule.cpp
+++ b/Engine/Source/Editor/TurnkeySupport/Private/TurnkeySupportModule.cpp
@@ -7,6 +7,7 @@
 #include "Widgets/Images/SImage.h"
 #include "Widgets/Layout/SBox.h"
 #include "Widgets/Input/SButton.h"
+#include "Widgets/SToolTip.h"
 #include "MessageLogModule.h"
 #include "Framework/Notifications/NotificationManager.h"
 #include "Widgets/Notifications/SNotificationList.h"
@@ -851,7 +852,7 @@ static FSlateIcon MakePlatformSdkIconAttribute(FName IniPlatformName, TSharedPtr
 //		));
 }
 
-static FText FormatSdkInfo(const FTurnkeySdkInfo& SdkInfo, bool bIncludeAutoSdk)
+static void FormatSdkInfo(const FString& PlatformOrDevice, const FTurnkeySdkInfo& SdkInfo, bool bIncludeAutoSdk, FText& OutInfo, FText& OutToolTip)
 {
 	FFormatOrderedArguments Args;
 	Args.Add(FText::FromString(SdkInfo.InstalledVersion));
@@ -859,6 +860,7 @@ static FText FormatSdkInfo(const FTurnkeySdkInfo& SdkInfo, bool bIncludeAutoSdk)
 	Args.Add(FText::FromString(SdkInfo.MinAllowedVersion));
 	Args.Add(FText::FromString(SdkInfo.MaxAllowedVersion));
 	Args.Add(SdkInfo.SdkErrorInformation);
+	Args.Add(FText::FromString(PlatformOrDevice));
 
 	TArray<FText> Lines;
 	Lines.Add(FText::Format(LOCTEXT("SdkInfo_Installed", "Installed SDK: {0}"), Args));
@@ -890,14 +892,24 @@ static FText FormatSdkInfo(const FTurnkeySdkInfo& SdkInfo, bool bIncludeAutoSdk)
 	}
 
 	// now make a single \n delimted text
-	return FText::Join(FText::FromString(TEXT("\n")), Lines);
+	OutInfo = FText::Join(FText::FromString(TEXT("\n")), Lines);
+
+	// make a tooltip
+	if (PlatformOrDevice.Contains(TEXT("@")))
+	{
+		OutToolTip = FText::Format(LOCTEXT("SdkInfo_ToolTip", "Information returned from:\nRunUAT Turnkey -command=VerifySdk -device={5}"), Args);
+	}
+	else
+	{
+		OutToolTip = FText::Format(LOCTEXT("SdkInfo_ToolTip", "Information returned from:\nRunUAT Turnkey -command=VerifySdk -platform={5}"), Args);
+	}
 }
 
 
 
 
 
-static void MakeTurnkeyPlatformMenu(FMenuBuilder& MenuBuilder, FName IniPlatformName, ITargetDeviceServicesModule* TargetDeviceServicesModule)
+static void MakeTurnkeyPlatformMenu(UToolMenu* ToolMenu, FName IniPlatformName, ITargetDeviceServicesModule* TargetDeviceServicesModule)
 {
 	const FDataDrivenPlatformInfo& DDPI = FDataDrivenPlatformInfoRegistry::GetPlatformInfo(IniPlatformName);
 	FString UBTPlatformString = DDPI.UBTPlatformString;
@@ -906,10 +918,10 @@ static void MakeTurnkeyPlatformMenu(FMenuBuilder& MenuBuilder, FName IniPlatform
 
 	if (VanillaInfo != nullptr)
 	{
+		FToolMenuSection& Section = ToolMenu->AddSection("ContentManagement", LOCTEXT("TurnkeySection_Content", "Content Management"));
 
-		MenuBuilder.BeginSection("ContentManagement", LOCTEXT("TurnkeySection_Content", "Content Management"));
-
-		MenuBuilder.AddMenuEntry(
+		Section.AddMenuEntry(
+			NAME_None,
 			LOCTEXT("Turnkey_PackageProject", "Package Project"),
 			LOCTEXT("TurnkeyTooltip_PackageProject", "Package this project and archive it to a user-selected directory. This can then be used to install and run."),
 			FSlateIcon(),
@@ -919,7 +931,8 @@ static void MakeTurnkeyPlatformMenu(FMenuBuilder& MenuBuilder, FName IniPlatform
 			)
 		);
 
-		MenuBuilder.AddMenuEntry(
+		Section.AddMenuEntry(
+			NAME_None,
 			LOCTEXT("Turnkey_CookContent", "Cook Content"),
 			LOCTEXT("TurnkeyTooltip_CookContent", "Cook this project for the selected configuration and target"),
 			FSlateIcon(),
@@ -929,7 +942,8 @@ static void MakeTurnkeyPlatformMenu(FMenuBuilder& MenuBuilder, FName IniPlatform
 			)
 		);
 // 
-// 		MenuBuilder.AddMenuEntry(
+// 		Section.AddMenuEntry(
+//			NAME_None,
 // 			LOCTEXT("Turnkey_PrepareForDebugging", "Prepare For Debugging"),
 // 			LOCTEXT("TurnkeyTooltip_PrepareForDebugging", "Prepare this project for debugging"),
 // 			FSlateIcon(),
@@ -947,7 +961,8 @@ static void MakeTurnkeyPlatformMenu(FMenuBuilder& MenuBuilder, FName IniPlatform
 		{
 			if (Build.SpecificPlatforms.Num() == 0 || Build.SpecificPlatforms.Contains(PlatformString))
 			{
-				MenuBuilder.AddMenuEntry(
+				Section.AddMenuEntry(
+					NAME_None,
 					FText::FromString(Build.Name),
 					// @todo turnkey: add the build string to the tooltip
 					LOCTEXT("TurnkeyTooltip_EngineCustomBuild", "Execute a custom build"),
@@ -964,7 +979,8 @@ static void MakeTurnkeyPlatformMenu(FMenuBuilder& MenuBuilder, FName IniPlatform
 		{
 			if (Build.SpecificPlatforms.Num() == 0 || Build.SpecificPlatforms.Contains(PlatformString))
 			{
-				MenuBuilder.AddMenuEntry(
+				Section.AddMenuEntry(
+					NAME_None,
 					FText::FromString(Build.Name),
 					// @todo turnkey: add the build string to the tooltip
 					LOCTEXT("TurnkeyTooltip_ProjectCustomBuild", "Execute a custom build (this comes from Packaging Settings)"),
@@ -977,8 +993,6 @@ static void MakeTurnkeyPlatformMenu(FMenuBuilder& MenuBuilder, FName IniPlatform
 			}
 		}
 
-		MenuBuilder.EndSection();
-
 		UProjectPackagingSettings* AllPlatformPackagingSettings = GetMutableDefault<UProjectPackagingSettings>();
 
 		// Flavor Selection exists for Android to be able to package ASTC, DXT, ETC2, etc
@@ -999,12 +1013,13 @@ static void MakeTurnkeyPlatformMenu(FMenuBuilder& MenuBuilder, FName IniPlatform
 				{
 					FTurnkeySupportCallbacks::SetActiveFlavor(ValidFlavors[0]);
 				}
-				
-				MenuBuilder.BeginSection("FlavorSelection", LOCTEXT("TurnkeySection_FlavorSelection", "Flavor Selection"));
+
+				FToolMenuSection& FlavorSection = ToolMenu->AddSection("FlavorSelection", LOCTEXT("TurnkeySection_FlavorSelection", "Flavor Selection"));
 				
 				for (const PlatformInfo::FTargetPlatformInfo* Info : ValidFlavors)
 				{
-					MenuBuilder.AddMenuEntry(
+					FlavorSection.AddMenuEntry(
+						NAME_None,
 						Info->DisplayName,
 						FText(),
 						FSlateIcon(),
@@ -1013,18 +1028,17 @@ static void MakeTurnkeyPlatformMenu(FMenuBuilder& MenuBuilder, FName IniPlatform
 							FCanExecuteAction::CreateStatic(&FTurnkeySupportCallbacks::CanSetActiveFlavor, Info),
 							FIsActionChecked::CreateStatic(&FTurnkeySupportCallbacks::SetActiveFlavorIsChecked, Info)
 						),
-						NAME_None,
 						EUserInterfaceActionType::RadioButton
 					);
 				}
-				MenuBuilder.EndSection();
 			}
 		}
 
-		MenuBuilder.BeginSection("BuildConfig", LOCTEXT("TurnkeySection_BuildConfig", "Binary Configuration"));
+		FToolMenuSection& ConfigSection = ToolMenu->AddSection("BuildConfig", LOCTEXT("TurnkeySection_BuildConfig", "Binary Configuration"));
 		
 		const UProjectPackagingSettings::FConfigurationInfo& ConfigInfo = UProjectPackagingSettings::ConfigurationInfo[static_cast<int32>(AllPlatformPackagingSettings->BuildConfiguration)];
-		MenuBuilder.AddMenuEntry(
+		ConfigSection.AddMenuEntry(
+			NAME_None,
 			FText::Format(LOCTEXT("DefaultConfiguration",  "Use Project Setting ({0})"), ConfigInfo.Name),
 			ConfigInfo.ToolTip,
 			FSlateIcon(),
@@ -1033,7 +1047,6 @@ static void MakeTurnkeyPlatformMenu(FMenuBuilder& MenuBuilder, FName IniPlatform
 				FCanExecuteAction(),
 				FIsActionChecked::CreateStatic(&FTurnkeySupportCallbacks::PackageBuildConfigurationIsChecked, VanillaInfo, EProjectPackagingBuildConfigurations::PPBC_MAX)
 			),
-			NAME_None,
 			EUserInterfaceActionType::RadioButton
 		);
 
@@ -1045,7 +1058,8 @@ static void MakeTurnkeyPlatformMenu(FMenuBuilder& MenuBuilder, FName IniPlatform
 			const UProjectPackagingSettings::FConfigurationInfo& ConfigurationInfo = UProjectPackagingSettings::ConfigurationInfo[(int)PackagingConfiguration];
 			if (FInstalledPlatformInfo::Get().IsValid(TOptional<EBuildTargetType>(), TOptional<FString>(), ConfigurationInfo.Configuration, ProjectType, EInstalledPlatformState::Downloaded))
 			{
-				MenuBuilder.AddMenuEntry(
+				ConfigSection.AddMenuEntry(
+					NAME_None,
 					ConfigurationInfo.Name,
 					ConfigurationInfo.ToolTip,
 					FSlateIcon(),
@@ -1054,12 +1068,10 @@ static void MakeTurnkeyPlatformMenu(FMenuBuilder& MenuBuilder, FName IniPlatform
 						FCanExecuteAction(),
 						FIsActionChecked::CreateStatic(&FTurnkeySupportCallbacks::PackageBuildConfigurationIsChecked, VanillaInfo, PackagingConfiguration)
 					),
-					NAME_None,
 					EUserInterfaceActionType::RadioButton
 				);
 			}
 		}
-		MenuBuilder.EndSection();
 
 		// Collect build targets. Content-only projects use Engine targets
 		FProjectStatus ProjectStatus;
@@ -1092,9 +1104,10 @@ static void MakeTurnkeyPlatformMenu(FMenuBuilder& MenuBuilder, FName IniPlatform
 					AllPlatformPackagingSettings->SaveConfig();
 				}
 
-				MenuBuilder.BeginSection("BuildTarget", LOCTEXT("TurnkeySection_BuildTarget", "Build Target"));
+				FToolMenuSection& TargetSection = ToolMenu->AddSection("BuildTarget", LOCTEXT("TurnkeySection_BuildTarget", "Build Target"));
 
-				MenuBuilder.AddMenuEntry(
+				TargetSection.AddMenuEntry(
+					NAME_None,
 					FText::Format(LOCTEXT("DefaultPackageTarget",  "Use Project Setting ({0})"), FText::FromString(AllPlatformPackagingSettings->BuildTarget)),
 					FText::Format(LOCTEXT("DefaultPackageTargetTooltip", "Package the {0} target"), FText::FromString(AllPlatformPackagingSettings->BuildTarget)),
 					FSlateIcon(),
@@ -1103,13 +1116,13 @@ static void MakeTurnkeyPlatformMenu(FMenuBuilder& MenuBuilder, FName IniPlatform
 						FCanExecuteAction(),
 						FIsActionChecked::CreateStatic(&FTurnkeySupportCallbacks::PackageBuildTargetIsChecked, VanillaInfo, FString(""))
 						),
-					NAME_None,
 					EUserInterfaceActionType::RadioButton
 				);
 
 				for (const FTargetInfo& Target : ValidTargets)
 				{
-					MenuBuilder.AddMenuEntry(
+					TargetSection.AddMenuEntry(
+						NAME_None,
 						FText::FromString(Target.Name),
 						FText::Format(LOCTEXT("PackageTargetName", "Package the '{0}' target."), FText::FromString(Target.Name)),
 						FSlateIcon(),
@@ -1118,16 +1131,13 @@ static void MakeTurnkeyPlatformMenu(FMenuBuilder& MenuBuilder, FName IniPlatform
 							FCanExecuteAction(),
 							FIsActionChecked::CreateStatic(&FTurnkeySupportCallbacks::PackageBuildTargetIsChecked, VanillaInfo, Target.Name)
 						),
-						NAME_None,
 						EUserInterfaceActionType::RadioButton
 					);
 				}
-				
-				MenuBuilder.EndSection();
 			}
 		}
 
-		MenuBuilder.BeginSection("AllDevices", LOCTEXT("TurnkeySection_AllDevices", "All Devices"));
+		FToolMenuSection& DevicesSection = ToolMenu->AddSection("AllDevices", LOCTEXT("TurnkeySection_AllDevices", "All Devices"));
 
 		TArray<TSharedPtr<ITargetDeviceProxy>> DeviceProxies;
 		TargetDeviceServicesModule->GetDeviceProxyManager()->GetAllProxies(IniPlatformName, DeviceProxies);
@@ -1136,23 +1146,32 @@ static void MakeTurnkeyPlatformMenu(FMenuBuilder& MenuBuilder, FName IniPlatform
 		{
 			FString DeviceName = Proxy->GetName();
 			FString DeviceId = Proxy->GetTargetDeviceId(NAME_None);
-			MenuBuilder.AddSubMenu(
+			DevicesSection.AddSubMenu(
+				NAME_None,
 				MakeSdkStatusAttribute(IniPlatformName, Proxy),
 				FText(),
-				FNewMenuDelegate::CreateLambda([IniPlatformName, DeviceName, DeviceId](FMenuBuilder& SubMenuBuilder)
+				FNewToolMenuDelegate::CreateLambda([IniPlatformName, DeviceName, DeviceId](UToolMenu* SubToolMenu)
 				{
 					FTurnkeySdkInfo SdkInfo = ITurnkeySupportModule::Get().GetSdkInfoForDeviceId(DeviceId);
-
-					SubMenuBuilder.AddWidget(
-						SNew(STextBlock)
-						.ColorAndOpacity(FSlateColor::UseSubduedForeground())
-						.Text(FormatSdkInfo(SdkInfo, false)),
-						FText::GetEmpty()
+					FText SdkText, SdkToolTip;
+					FormatSdkInfo(DeviceId, SdkInfo, false, SdkText, SdkToolTip);
+
+					FToolMenuSection& Section = SubToolMenu->AddSection(NAME_None);
+					Section.AddEntry(
+						FToolMenuEntry::InitWidget(
+							NAME_None,
+							SNew(STextBlock)
+							.ColorAndOpacity(FSlateColor::UseSubduedForeground())
+							.Text(SdkText)
+							.ToolTip(SNew(SToolTip).Text(SdkToolTip)),
+							FText::GetEmpty()
+						)
 					);
 
 					if (SdkInfo.DeviceStatus == ETurnkeyDeviceStatus::SoftwareValid)
 					{
-						SubMenuBuilder.AddMenuEntry(
+						Section.AddMenuEntry(
+							NAME_None,
 							LOCTEXT("Turnkey_ForceRepairDevice", "Force Update Device"),
 							LOCTEXT("TurnkeyTooltip_ForceRepairDevice", "Force repairing anything on the device needed (update firmware, etc). Will perform all steps possible, even if not needed."),
 							FSlateIcon(),
@@ -1161,7 +1180,8 @@ static void MakeTurnkeyPlatformMenu(FMenuBuilder& MenuBuilder, FName IniPlatform
 					}
 					else
 					{
-						SubMenuBuilder.AddMenuEntry(
+						Section.AddMenuEntry(
+							NAME_None,
 							LOCTEXT("Turnkey_RepairDevice", "Update Device"),
 							LOCTEXT("TurnkeyTooltip_RepairDevice", "Perform any fixup that may be needed on this device. If up to date already, nothing will be done."),
 							FSlateIcon(),
@@ -1173,39 +1193,27 @@ static void MakeTurnkeyPlatformMenu(FMenuBuilder& MenuBuilder, FName IniPlatform
 				false, MakePlatformSdkIconAttribute(IniPlatformName, Proxy)
 			);
 		}
-
-
-		MenuBuilder.EndSection();
 	}
 
-	MenuBuilder.BeginSection("SdkManagement", LOCTEXT("TurnkeySection_Sdks", "Sdk Managment"));
+	FToolMenuSection& SdkSection = ToolMenu->AddSection("SdkManagement", LOCTEXT("TurnkeySection_Sdks", "Sdk Managment"));
 
 	const FTurnkeySdkInfo& SdkInfo = ITurnkeySupportModule::Get().GetSdkInfo(IniPlatformName, true);
-	FFormatOrderedArguments Args;
-	Args.Add(FText::FromString(SdkInfo.InstalledVersion));
-	Args.Add(FText::FromString(SdkInfo.AutoSDKVersion));
-	Args.Add(FText::FromString(SdkInfo.MinAllowedVersion));
-	Args.Add(FText::FromString(SdkInfo.MaxAllowedVersion));
-	if (SdkInfo.SdkErrorInformation.IsEmpty())
-	{
-		Args.Add(FText::GetEmpty());
-		Args.Add(FText::GetEmpty());
-	}
-	else
-	{
-		Args.Add(LOCTEXT("ErrorPrefix", "\nErrors:\n"));
-		Args.Add(FText::GetEmpty());
-	}
+	FText SdkText, SdkToolTip;
+	FormatSdkInfo(IniPlatformName.ToString(), SdkInfo, true, SdkText, SdkToolTip);
 
-	MenuBuilder.AddWidget(
-		SNew(SBox)
-		.Padding(FMargin(16.0f, 3.0f))
-		[
-			SNew(STextBlock)
-			.ColorAndOpacity(FSlateColor::UseSubduedForeground())
-			.Text(FormatSdkInfo(SdkInfo, true))
-		],
-		FText::GetEmpty()
+	SdkSection.AddEntry(
+		FToolMenuEntry::InitWidget(
+			NAME_None,
+			SNew(SBox)
+			.Padding(FMargin(16.0f, 3.0f))
+			[
+				SNew(STextBlock)
+				.ColorAndOpacity(FSlateColor::UseSubduedForeground())
+				.Text(SdkText)
+				.ToolTip(SNew(SToolTip).Text(SdkToolTip))
+			],
+			FText::GetEmpty()
+		)
 	);
 
 	FString NoDevice;
@@ -1213,7 +1221,8 @@ static void MakeTurnkeyPlatformMenu(FMenuBuilder& MenuBuilder, FName IniPlatform
 	{
 		if (SdkInfo.Status == ETurnkeyPlatformSdkStatus::OutOfDate)
 		{
-			MenuBuilder.AddMenuEntry(
+			SdkSection.AddMenuEntry(
+				NAME_None,
 				LOCTEXT("Turnkey_UpdateSdkMinimal", "Update Sdk"),
 				LOCTEXT("TurnkeyTooltip_InstallSdkMinimal", "Attempt to update an Sdk, as hosted by your studio. Will attempt to install a minimal Sdk (useful for building/running only)"),
 				FSlateIcon(),
@@ -1222,7 +1231,8 @@ static void MakeTurnkeyPlatformMenu(FMenuBuilder& MenuBuilder, FName IniPlatform
 
 			if (SdkInfo.bCanInstallFullSdk && SdkInfo.bCanInstallAutoSdk)
 			{
-				MenuBuilder.AddMenuEntry(
+				SdkSection.AddMenuEntry(
+					NAME_None,
 					LOCTEXT("Turnkey_UpdateSdkFull", "Update Sdk (Full Platform Installer)"),
 					LOCTEXT("TurnkeyTooltip_UpdateSdkFull", "Attempt to update an Sdk, as hosted by your studio. Will attempt to install a full Sdk (useful profiling or other use cases)"),
 					FSlateIcon(),
@@ -1232,7 +1242,8 @@ static void MakeTurnkeyPlatformMenu(FMenuBuilder& MenuBuilder, FName IniPlatform
 		}
 		else if (SdkInfo.Status == ETurnkeyPlatformSdkStatus::Valid)
 		{
-			MenuBuilder.AddMenuEntry(
+			SdkSection.AddMenuEntry(
+				NAME_None,
 				LOCTEXT("Turnkey_ForceSdkMinimal", "Force Reinstall Sdk"),
 				LOCTEXT("TurnkeyTooltip_ForceSdkMinimal", "Attempt to force re-install an Sdk, as hosted by your studio. Will attempt to install a minimal Sdk (useful for building/running only)"),
 				FSlateIcon(),
@@ -1241,7 +1252,8 @@ static void MakeTurnkeyPlatformMenu(FMenuBuilder& MenuBuilder, FName IniPlatform
 
 			if (SdkInfo.bCanInstallFullSdk && SdkInfo.bCanInstallAutoSdk)
 			{
-				MenuBuilder.AddMenuEntry(
+				SdkSection.AddMenuEntry(
+					NAME_None,
 					LOCTEXT("Turnkey_ForceSdkFull", "Force Reinstall (Full Platform Installer)"),
 					LOCTEXT("TurnkeyTooltip_ForceSdkForce", "Attempt to force re-install an Sdk, as hosted by your studio. Will attempt to install a full Sdk (useful profiling or other use cases)"),
 					FSlateIcon(),
@@ -1251,14 +1263,16 @@ static void MakeTurnkeyPlatformMenu(FMenuBuilder& MenuBuilder, FName IniPlatform
 		}
 		else
 		{
-			MenuBuilder.AddMenuEntry(
+			SdkSection.AddMenuEntry(
+				NAME_None,
 				LOCTEXT("Turnkey_InstallSdkMinimal", "Install Sdk"),
 				LOCTEXT("TurnkeyTooltip_InstallSdkMinimal", "Attempt to install an Sdk, as hosted by your studio. Will attempt to install a minimal Sdk (useful for building/running only)"),
 				FSlateIcon(),
 				FExecuteAction::CreateStatic(TurnkeyInstallSdk, IniPlatformName.ToString(), false, false, NoDevice)
 			);
 
-			MenuBuilder.AddMenuEntry(
+			SdkSection.AddMenuEntry(
+				NAME_None,
 				LOCTEXT("Turnkey_InstallSdkFull", "Install Sdk (Full Platform Installer)"),
 				LOCTEXT("TurnkeyTooltip_InstallSdkFull", "Attempt to install an Sdk, as hosted by your studio. Will attempt to install a full Sdk (useful profiling or other use cases)"),
 				FSlateIcon(),
@@ -1269,7 +1283,8 @@ static void MakeTurnkeyPlatformMenu(FMenuBuilder& MenuBuilder, FName IniPlatform
 	else
 	{
 		// if Turnkey can't be used for this platform, then show old-school documentation
-		MenuBuilder.AddMenuEntry(
+		SdkSection.AddMenuEntry(
+			NAME_None,
 			LOCTEXT("Turnkey_ShowDocumentation", "Installation Help..."),
 			LOCTEXT("TurnkeyTooltip_ShowDocumentation", "Show documentation with help installing the SDK for this platform"),
 			FSlateIcon(),
@@ -1490,8 +1505,10 @@ void FTurnkeySupportModule::MakeQuickLaunchItems(class UToolMenu* Menu, FOnQuick
 						NAME_None,
 						MakeSdkStatusAttribute(PlatformName, DeviceProxies[0]),
 						Tooltip,
-						FNewMenuDelegate::CreateLambda([TargetDeviceServicesModule, PlatformName, ExternalOnClickDelegate](FMenuBuilder& SubMenuBuilder)
+						FNewToolMenuDelegate::CreateLambda([TargetDeviceServicesModule, PlatformName, ExternalOnClickDelegate](UToolMenu* SubToolMenu)
 							{
+								FToolMenuSection& Section = SubToolMenu->AddSection(NAME_None);
+
 								// re-get the proxies, just in case they changed
 								TArray<TSharedPtr<ITargetDeviceProxy>> DeviceProxies;
 								TargetDeviceServicesModule->GetDeviceProxyManager()->GetAllProxies(PlatformName, DeviceProxies);
@@ -1501,12 +1518,12 @@ void FTurnkeySupportModule::MakeQuickLaunchItems(class UToolMenu* Menu, FOnQuick
 									FUIAction SubAction;
 									FText SubTooltip;
 									GenerateDeviceProxyMenuParams(Proxy, PlatformName, SubAction, SubTooltip, ExternalOnClickDelegate);
-									SubMenuBuilder.AddMenuEntry(
+									Section.AddMenuEntry(
+										NAME_None,
 										MakeSdkStatusAttribute(PlatformName, Proxy),
 										SubTooltip,
 										MakePlatformSdkIconAttribute(PlatformName, Proxy),
 										SubAction,
-										NAME_None,
 										EUserInterfaceActionType::Button
 									);
 								}
@@ -1597,7 +1614,6 @@ TSharedRef<SWidget> FTurnkeySupportModule::MakeTurnkeyMenuWidget() const
 	const FTurnkeySupportCommands& Commands = FTurnkeySupportCommands::Get();
 
 	const bool bShouldCloseWindowAfterMenuSelection = true;
-//	FMenuBuilder MenuBuilder(bShouldCloseWindowAfterMenuSelection, FTurnkeySupportCommands::ActionList);
 
 	static const FName MenuName("UnrealEd.PlayWorldCommands.PlatformsMenu");
 
@@ -1652,7 +1668,7 @@ TSharedRef<SWidget> FTurnkeySupportModule::MakeTurnkeyMenuWidget() const
 					NAME_None,
 					MakeSdkStatusAttribute(PlatformName, nullptr),
 					FText::FromString(PlatformName.ToString()),
-					FNewMenuDelegate::CreateStatic(&MakeTurnkeyPlatformMenu, PlatformName, TargetDeviceServicesModule),
+					FNewToolMenuDelegate::CreateStatic(&MakeTurnkeyPlatformMenu, PlatformName, TargetDeviceServicesModule),
 					false,
 					MakePlatformSdkIconAttribute(PlatformName, nullptr),
 					true
@@ -1677,14 +1693,16 @@ TSharedRef<SWidget> FTurnkeySupportModule::MakeTurnkeyMenuWidget() const
 					NAME_None,
 					LOCTEXT("Turnkey_UnsupportedPlatforms", "Platforms Not Supported by Project"),
 					LOCTEXT("Turnkey_UnsupportedPlatformsToolTip", "List of platforms that are not marked as supported by this platform. Use the \"Supported Platforms...\""),
-					FNewMenuDelegate::CreateLambda([UnsupportedPlatforms, TargetDeviceServicesModule](FMenuBuilder& SubMenuBuilder)
+					FNewToolMenuDelegate::CreateLambda([UnsupportedPlatforms, TargetDeviceServicesModule](UToolMenu* SubToolMenu)
 						{
+							FToolMenuSection& Section = SubToolMenu->AddSection(NAME_None);
 							for (const auto It : UnsupportedPlatforms)
 							{
-								SubMenuBuilder.AddSubMenu(
+								Section.AddSubMenu(
+									NAME_None,
 									MakeSdkStatusAttribute(It.Key, nullptr),
 									FText::FromString(It.Key.ToString()),
-									FNewMenuDelegate::CreateStatic(&MakeTurnkeyPlatformMenu, It.Key, TargetDeviceServicesModule),
+									FNewToolMenuDelegate::CreateStatic(&MakeTurnkeyPlatformMenu, It.Key, TargetDeviceServicesModule),
 									false,
 									MakePlatformSdkIconAttribute(It.Key, nullptr),
 									true
@@ -1702,14 +1720,16 @@ TSharedRef<SWidget> FTurnkeySupportModule::MakeTurnkeyMenuWidget() const
 					NAME_None,
 					LOCTEXT("Turnkey_UncompiledPlatforms", "Platforms With No Compiled Support"),
 					LOCTEXT("Turnkey_UncompiledPlatformsToolTip", "List of platforms that you have access to, but support is not compiled in to the editor. It may be caused by missing an SDK, so you attempt to install an SDK here."),
-					FNewMenuDelegate::CreateLambda([UncompiledPlatforms, TargetDeviceServicesModule](FMenuBuilder& SubMenuBuilder)
+					FNewToolMenuDelegate::CreateLambda([UncompiledPlatforms, TargetDeviceServicesModule](UToolMenu* SubToolMenu)
 						{
+							FToolMenuSection& Section = SubToolMenu->AddSection(NAME_None);
 							for (const auto It : UncompiledPlatforms)
 							{
-								SubMenuBuilder.AddSubMenu(
+								Section.AddSubMenu(
+									NAME_None,
 									MakeSdkStatusAttribute(It.Key, nullptr),
 									FText::FromString(It.Key.ToString()),
-									FNewMenuDelegate::CreateStatic(&MakeTurnkeyPlatformMenu, It.Key, TargetDeviceServicesModule),
+									FNewToolMenuDelegate::CreateStatic(&MakeTurnkeyPlatformMenu, It.Key, TargetDeviceServicesModule),
 									false,
 									MakePlatformSdkIconAttribute(It.Key, nullptr),
 									true
