Disable sdk installation documentation menu item until we have URLs to go with

diff --git a/Engine/Source/Editor/TurnkeySupport/Private/TurnkeySupportModule.cpp b/Engine/Source/Editor/TurnkeySupport/Private/TurnkeySupportModule.cpp
--- a/Engine/Source/Editor/TurnkeySupport/Private/TurnkeySupportModule.cpp
+++ b/Engine/Source/Editor/TurnkeySupport/Private/TurnkeySupportModule.cpp
@@ -1280,17 +1280,17 @@ static void MakeTurnkeyPlatformMenu(UToolMenu* ToolMenu, FName IniPlatformName,
 			);
 		}
 	}
-	else
-	{
-		// if Turnkey can't be used for this platform, then show old-school documentation
-		SdkSection.AddMenuEntry(
-			NAME_None,
-			LOCTEXT("Turnkey_ShowDocumentation", "Installation Help..."),
-			LOCTEXT("TurnkeyTooltip_ShowDocumentation", "Show documentation with help installing the SDK for this platform"),
-			FSlateIcon(),
-			FExecuteAction::CreateStatic(ShowInstallationHelp, IniPlatformName)
-		);
-	}
+
+#if 0 // @todo turnkey enable and always show documentation for installation help when we have URLs in place
+	// Link to documentation
+	SdkSection.AddMenuEntry(
+		NAME_None,
+		LOCTEXT("Turnkey_ShowDocumentation", "Installation Help..."),
+		LOCTEXT("TurnkeyTooltip_ShowDocumentation", "Show documentation with help installing the SDK for this platform"),
+		FSlateIcon(),
+		FExecuteAction::CreateStatic(ShowInstallationHelp, IniPlatformName)
+	);
+#endif
 }
 
 
