- Changed FMonitoredProcess to buffer the output if the OnOutput delegate is not bound (get the buffered output with GetFullOutputWithoutDelegate())
- Log out the Turnkey UAT output on any failure as a warning

diff --git a/Engine/Source/Editor/TurnkeySupport/Private/TurnkeySupportModule.cpp b/Engine/Source/Editor/TurnkeySupport/Private/TurnkeySupportModule.cpp
--- a/Engine/Source/Editor/TurnkeySupport/Private/TurnkeySupportModule.cpp
+++ b/Engine/Source/Editor/TurnkeySupport/Private/TurnkeySupportModule.cpp
@@ -2168,10 +2168,9 @@ void FTurnkeySupportModule::UpdateSdkInfo()
 				for (auto& It : PerPlatformSdkInfo)
 				{
 					It.Value.Status = ETurnkeyPlatformSdkStatus::Error;
-					It.Value.SdkErrorInformation = FText::Format(NSLOCTEXT("Turnkey", "TurnkeyError_ReturnedError", "Turnkey returned an error, code {0}"), { ExitCode });
-
-					// @todo turnkey error description!
+					It.Value.SdkErrorInformation = FText::Format(NSLOCTEXT("Turnkey", "TurnkeyError_ReturnedError", "Turnkey returned an error, code {0} (See log)"), { ExitCode });
 				}
+				UE_LOG(LogTurnkeySupport, Warning, TEXT("Turnkey failed to run properly, full Turnkey output:\n%s"), *TurnkeyProcess->GetFullOutputWithoutDelegate());
 			}
 
 
@@ -2269,18 +2268,29 @@ void FTurnkeySupportModule::UpdateSdkInfoForDevices(TArray<FString> PlatformDevi
 						PerDeviceSdkInfo[DDPIDeviceId] = SdkInfo;
 					}
 				}
-			}
 
-			for (const FString& Id : PlatformDeviceIds)
+			    for (const FString& Id : PlatformDeviceIds)
+			    {
+				    FTurnkeySdkInfo& SdkInfo = PerDeviceSdkInfo[ConvertToDDPIDeviceId(Id)];
+				    if (SdkInfo.Status == ETurnkeyPlatformSdkStatus::Querying)
+				    {
+					    SdkInfo.Status = ETurnkeyPlatformSdkStatus::Error;
+					    SdkInfo.SdkErrorInformation = NSLOCTEXT("Turnkey", "TurnkeyError_DeviceNotReturned", "A device's Sdk status was not returned from Turnkey");
+				    }
+			    }
+			}
+			else
 			{
-				FTurnkeySdkInfo& SdkInfo = PerDeviceSdkInfo[ConvertToDDPIDeviceId(Id)];
-				if (SdkInfo.Status == ETurnkeyPlatformSdkStatus::Querying)
-				{
-					SdkInfo.Status = ETurnkeyPlatformSdkStatus::Error;
-					SdkInfo.SdkErrorInformation = NSLOCTEXT("Turnkey", "TurnkeyError_DeviceNotReturned", "A device's Sdk status was not returned from Turnkey");
+			    for (const FString& Id : PlatformDeviceIds)
+			    {
+				    FTurnkeySdkInfo& SdkInfo = PerDeviceSdkInfo[ConvertToDDPIDeviceId(Id)];
+				    SdkInfo.Status = ETurnkeyPlatformSdkStatus::Error;
+					SdkInfo.SdkErrorInformation = FText::Format(NSLOCTEXT("Turnkey", "TurnkeyError_ReturnedError", "Turnkey returned an error, code {0} (See log)"), { ExitCode });
 				}
+				UE_LOG(LogTurnkeySupport, Warning, TEXT("Turnkey failed to run properly, full Turnkey output:\n%s"), *TurnkeyProcess->GetFullOutputWithoutDelegate());
 			}
 
+
 			// cleanup
 			delete TurnkeyProcess;
 			IFileManager::Get().Delete(*ReportFilename);
diff --git a/Engine/Source/Runtime/Core/Private/Misc/MonitoredProcess.cpp b/Engine/Source/Runtime/Core/Private/Misc/MonitoredProcess.cpp
--- a/Engine/Source/Runtime/Core/Private/Misc/MonitoredProcess.cpp
+++ b/Engine/Source/Runtime/Core/Private/Misc/MonitoredProcess.cpp
@@ -101,25 +101,29 @@ void FMonitoredProcess::ProcessOutput( const FString& Output )
 	// Append this output to the output buffer
 	OutputBuffer += Output;
 
-	// Output all the complete lines
-	int32 LineStartIdx = 0;
-	for(int32 Idx = 0; Idx < OutputBuffer.Len(); Idx++)
+	// if the delegate is not bound, then just keep buffering the output to OutputBuffer for later return from GetFullOutputWithoutDelegate()
+	if (OutputDelegate.IsBound())
 	{
-		if(OutputBuffer[Idx] == '\r' || OutputBuffer[Idx] == '\n')
+		// Output all the complete lines
+		int32 LineStartIdx = 0;
+		for (int32 Idx = 0; Idx < OutputBuffer.Len(); Idx++)
 		{
-			OutputDelegate.ExecuteIfBound(OutputBuffer.Mid(LineStartIdx, Idx - LineStartIdx));
-			
-			if(OutputBuffer[Idx] == '\r' && Idx + 1 < OutputBuffer.Len() && OutputBuffer[Idx + 1] == '\n')
+			if (OutputBuffer[Idx] == '\r' || OutputBuffer[Idx] == '\n')
 			{
-				Idx++;
-			}
+				OutputDelegate.ExecuteIfBound(OutputBuffer.Mid(LineStartIdx, Idx - LineStartIdx));
+
+				if (OutputBuffer[Idx] == '\r' && Idx + 1 < OutputBuffer.Len() && OutputBuffer[Idx + 1] == '\n')
+				{
+					Idx++;
+				}
 
-			LineStartIdx = Idx + 1;
+				LineStartIdx = Idx + 1;
+			}
 		}
-	}
 
-	// Remove all the complete lines from the buffer
-	OutputBuffer.MidInline(LineStartIdx, MAX_int32, false);
+		// Remove all the complete lines from the buffer
+		OutputBuffer.MidInline(LineStartIdx, MAX_int32, false);
+	}
 }
 
 void FMonitoredProcess::TickInternal()
diff --git a/Engine/Source/Runtime/Core/Public/Misc/MonitoredProcess.h b/Engine/Source/Runtime/Core/Public/Misc/MonitoredProcess.h
--- a/Engine/Source/Runtime/Core/Public/Misc/MonitoredProcess.h
+++ b/Engine/Source/Runtime/Core/Public/Misc/MonitoredProcess.h
@@ -150,6 +150,15 @@ public:
 		return ReturnCode;
 	}
 
+	/**
+	 * Returns the full output, wihtout needing to hookup a delegate and buffer it externally. Note that if OutputDelegate is 
+	 * bound, this _will not_ have the entire output
+	 */
+	const FString& GetFullOutputWithoutDelegate() const
+	{
+		return OutputBuffer;
+	}
+
 public:
 
 	// FRunnable interface
