- Run Turnkey device queries automtically for all devices after getting platform SDK info, so there is less waiting with Turnkey menu open with (Querying...)
- Added -nocompileuat when UAT has run successfully to speed up UAT runs

diff --git a/Engine/Source/Editor/TurnkeySupport/Private/TurnkeySupportModule.cpp b/Engine/Source/Editor/TurnkeySupport/Private/TurnkeySupportModule.cpp
--- a/Engine/Source/Editor/TurnkeySupport/Private/TurnkeySupportModule.cpp
+++ b/Engine/Source/Editor/TurnkeySupport/Private/TurnkeySupportModule.cpp
@@ -1654,7 +1654,6 @@ void FTurnkeySupportModule::MakeQuickLaunchItems(class UToolMenu* Menu, FOnQuick
 
 				if (DeviceProxies.Num() == 1)
 				{
-					UE_LOG(LogTurnkeySupport, Display, TEXT("Adding device menu item for %s"), *DeviceProxies[0]->GetName());
 					DynamicSection.AddMenuEntry(
 						NAME_None,
 						MakeSdkStatusAttribute(PlatformName, DeviceProxies[0]),
@@ -2162,6 +2161,14 @@ void FTurnkeySupportModule::UpdateSdkInfo()
 // 							*SdkInfo.AutoSDKVersion, *SdkInfo.MinAllowedVersion, *SdkInfo.MaxAllowedVersion);
 					}
 				}
+
+				// update all deviecs
+				UpdateSdkInfoForAllDevices();
+
+				// not all devices may be added yet, so add a delegate to catch any new devices that come along
+				ITargetDeviceServicesModule* TargetDeviceServicesModule = static_cast<ITargetDeviceServicesModule*>(FModuleManager::Get().LoadModule(TEXT("TargetDeviceServices")));
+				TargetDeviceServicesModule->GetDeviceProxyManager()->OnProxyAdded().RemoveAll(this);
+				TargetDeviceServicesModule->GetDeviceProxyManager()->OnProxyAdded().AddRaw(this, &FTurnkeySupportModule::UpdateSdkInfoForProxy);
 			}
 			else
 			{
@@ -2203,8 +2210,71 @@ void FTurnkeySupportModule::UpdateSdkInfo()
 	TurnkeyProcess->Launch();
 }
 
+void FTurnkeySupportModule::UpdateSdkInfoForAllDevices()
+{
+	TArray<FString> HostDevices;
+	TArray<FString> OtherDevices;
+
+	{
+		FScopeLock Lock(&GTurnkeySection);
+
+		// now kick off status update for all devices (host platform first, then everything else)
+		ITargetDeviceServicesModule* TargetDeviceServicesModule = static_cast<ITargetDeviceServicesModule*>(FModuleManager::Get().LoadModule(TEXT("TargetDeviceServices")));
+		for (const auto& Pair : FDataDrivenPlatformInfoRegistry::GetAllPlatformInfos())
+		{
+			FName PlatformName = Pair.Key;
+			if (!Pair.Value.bIsFakePlatform && !FDataDrivenPlatformInfoRegistry::IsPlatformHiddenFromUI(PlatformName))
+			{
+				// look for devices for all platforms, even if the platform isn't installed - Turnkey can install Sdk after selecting LaunchOn
+				TArray<TSharedPtr<ITargetDeviceProxy>> DeviceProxies;
+				TargetDeviceServicesModule->GetDeviceProxyManager()->GetAllProxies(PlatformName, DeviceProxies);
+
+				for (const TSharedPtr<ITargetDeviceProxy>& Proxy : DeviceProxies)
+				{
+					FString DeviceId = Proxy->GetTargetDeviceId(NAME_None);
+					if (GetSdkInfoForDeviceId(DeviceId).Status == ETurnkeyPlatformSdkStatus::Unknown)
+					{
+						if (Pair.Value.IniPlatformName == FPlatformProperties::IniPlatformName())
+						{
+							HostDevices.Add(DeviceId);
+						}
+						else
+						{
+							OtherDevices.Add(DeviceId);
+						}
+					}
+				}
+			}
+		}
+	}
+
+	UpdateSdkInfoForDevices(HostDevices);
+	UpdateSdkInfoForDevices(OtherDevices);
+}
+
+void FTurnkeySupportModule::UpdateSdkInfoForProxy(const TSharedRef<ITargetDeviceProxy>& AddedProxy)
+{
+	bool bIsNeeded;
+
+	FString DeviceId = AddedProxy->GetTargetDeviceId(NAME_None);
+	{
+		FScopeLock Lock(&GTurnkeySection);
+		bIsNeeded = GetSdkInfoForDeviceId(DeviceId).Status == ETurnkeyPlatformSdkStatus::Unknown;
+	}
+
+	if (bIsNeeded)
+	{
+		UpdateSdkInfoForDevices({ DeviceId });
+	}
+}
+
 void FTurnkeySupportModule::UpdateSdkInfoForDevices(TArray<FString> PlatformDeviceIds)
 {
+	if (PlatformDeviceIds.Num() == 0)
+	{
+		return;
+	}
+
 	FString BaseCommandline, ReportFilename;
 	PrepForTurnkeyReport(BaseCommandline, ReportFilename);
 
@@ -2390,4 +2460,4 @@ void FTurnkeySupportModule::ShutdownModule( )
 
 IMPLEMENT_MODULE(FTurnkeySupportModule, TurnkeySupport);
 
-#undef LOCTEXT_NAMESPACE
+#undef LOCTEXT_NAMESPACE
\ No newline at end of file
diff --git a/Engine/Source/Editor/TurnkeySupport/Private/TurnkeySupportModule.h b/Engine/Source/Editor/TurnkeySupport/Private/TurnkeySupportModule.h
--- a/Engine/Source/Editor/TurnkeySupport/Private/TurnkeySupportModule.h
+++ b/Engine/Source/Editor/TurnkeySupport/Private/TurnkeySupportModule.h
@@ -25,6 +25,8 @@ public:
 
 
 	virtual void UpdateSdkInfo() override;
+	void UpdateSdkInfoForAllDevices();
+	void UpdateSdkInfoForProxy(const TSharedRef<class ITargetDeviceProxy>& AddedProxy);
 	virtual void UpdateSdkInfoForDevices(TArray<FString> DeviceIds) override;
 
 	virtual FTurnkeySdkInfo GetSdkInfo(FName PlatformName, bool bBlockIfQuerying = true) const override;
diff --git a/Engine/Source/Runtime/Core/Private/Misc/MonitoredProcess.cpp b/Engine/Source/Runtime/Core/Private/Misc/MonitoredProcess.cpp
--- a/Engine/Source/Runtime/Core/Private/Misc/MonitoredProcess.cpp
+++ b/Engine/Source/Runtime/Core/Private/Misc/MonitoredProcess.cpp
@@ -303,7 +303,8 @@ bool FSerializedUATProcess::LaunchInternal()
 {
 	if (bHasSucceededOnce)
 	{
-		Params += TEXT(" -nocompile");
+		Params += TEXT(" -nocompile -nocompileuat");
+//		Params += TEXT(" -nocompile");
 	}
 
 	FOnMonitoredProcessCompleted OriginalCompletedDelegate = CompletedDelegate;
