- Replaced the TargetDeviceServices to get device proxies, and just use the TargetPlatforms directly, has the benefit of not needing to wait for the messagebux to process messages before getting access to the devices, so we can start querying during startup earlier

diff --git a/Engine/Source/Editor/TurnkeySupport/Private/TurnkeySupportModule.cpp b/Engine/Source/Editor/TurnkeySupport/Private/TurnkeySupportModule.cpp
--- a/Engine/Source/Editor/TurnkeySupport/Private/TurnkeySupportModule.cpp
+++ b/Engine/Source/Editor/TurnkeySupport/Private/TurnkeySupportModule.cpp
@@ -2194,20 +2194,11 @@ void FTurnkeySupportModule::UpdateSdkInfo()
 								PerPlatformSdkInfo[PlatformName].SDKVersions.Add(TEXT("AutoSDK"), SdkInfo.SDKVersions[TEXT("AutoSDK")]);
 							}
 						}
-
-
-// 						UE_LOG(LogTurnkeySupport, Log, TEXT("[TEST] Turnkey Platform: %s - %d, Installed: %s, AudoSDK: %s, Allowed: %s-%s"), *PlatformName.ToString(), (int)SdkInfo.Status, *SdkInfo.InstalledVersion,
-// 							*SdkInfo.AutoSDKVersion, *SdkInfo.MinAllowedVersion, *SdkInfo.MaxAllowedVersion);
 					}
 				}
 
 				// update all deviecs
 				UpdateSdkInfoForAllDevices();
-
-				// not all devices may be added yet, so add a delegate to catch any new devices that come along
-				ITargetDeviceServicesModule* TargetDeviceServicesModule = static_cast<ITargetDeviceServicesModule*>(FModuleManager::Get().LoadModule(TEXT("TargetDeviceServices")));
-				TargetDeviceServicesModule->GetDeviceProxyManager()->OnProxyAdded().RemoveAll(this);
-				TargetDeviceServicesModule->GetDeviceProxyManager()->OnProxyAdded().AddRaw(this, &FTurnkeySupportModule::UpdateSdkInfoForProxy);
 			}
 			else
 			{
@@ -2264,13 +2255,18 @@ void FTurnkeySupportModule::UpdateSdkInfoForAllDevices()
 			FName PlatformName = Pair.Key;
 			if (!Pair.Value.bIsFakePlatform && !FDataDrivenPlatformInfoRegistry::IsPlatformHiddenFromUI(PlatformName))
 			{
-				// look for devices for all platforms, even if the platform isn't installed - Turnkey can install Sdk after selecting LaunchOn
-				TArray<TSharedPtr<ITargetDeviceProxy>> DeviceProxies;
-				TargetDeviceServicesModule->GetDeviceProxyManager()->GetAllProxies(PlatformName, DeviceProxies);
+				ITargetPlatform* TP = GetTargetPlatformManager()->FindTargetPlatform(PlatformName);
+				if (TP == nullptr)
+				{
+					continue;
+				}
 
-				for (const TSharedPtr<ITargetDeviceProxy>& Proxy : DeviceProxies)
+				TArray<ITargetDevicePtr> Devices;
+				TP->GetAllDevices(Devices);
+
+				for (ITargetDevicePtr Device : Devices)
 				{
-					FString DeviceId = Proxy->GetTargetDeviceId(NAME_None);
+					FString DeviceId = Device->GetId().ToString();
 					if (GetSdkInfoForDeviceId(DeviceId).Status == ETurnkeyPlatformSdkStatus::Unknown)
 					{
 						if (Pair.Value.IniPlatformName == FPlatformProperties::IniPlatformName())
