- Check for bad device coming from a TargetPlatform in FTurnkeySupportModule::UpdateSdkInfoForAllDevices (and added a log because it would be good to know, although it's certainly not fatal)

diff --git a/Engine/Source/Editor/TurnkeySupport/Private/TurnkeySupportModule.cpp b/Engine/Source/Editor/TurnkeySupport/Private/TurnkeySupportModule.cpp
--- a/Engine/Source/Editor/TurnkeySupport/Private/TurnkeySupportModule.cpp
+++ b/Engine/Source/Editor/TurnkeySupport/Private/TurnkeySupportModule.cpp
@@ -2266,6 +2266,11 @@ void FTurnkeySupportModule::UpdateSdkInfoForAllDevices()
 
 				for (ITargetDevicePtr Device : Devices)
 				{
+					if (!Device.IsValid())
+					{
+						UE_LOG(LogTurnkeySupport, Log, TEXT("Platform %s returned an invlid device from GetAllDevices, which is not expected"), *PlatformName.ToString());
+						continue;
+					}
 					FString DeviceId = Device->GetId().ToString();
 					if (GetSdkInfoForDeviceId(DeviceId).Status == ETurnkeyPlatformSdkStatus::Unknown)
 					{
