Fix some LOCTEXT duplicate macros

diff --git a/Engine/Source/Editor/TurnkeySupport/Private/TurnkeySupportModule.cpp b/Engine/Source/Editor/TurnkeySupport/Private/TurnkeySupportModule.cpp
--- a/Engine/Source/Editor/TurnkeySupport/Private/TurnkeySupportModule.cpp
+++ b/Engine/Source/Editor/TurnkeySupport/Private/TurnkeySupportModule.cpp
@@ -631,7 +631,7 @@ public:
 			FString LogFilename, ReportFilename;
 			CommandLine += GetLogAndReportCommandline(LogFilename, ReportFilename);
 
-			FTurnkeyEditorSupport::RunUAT(CommandLine, PlatformInfo->DisplayName, LOCTEXT("Turnkey_GettingCommandLine", "Copying Commandline"), LOCTEXT("Turnkey_CustomTaskName", "CustomCommandLine"), nullptr /* TaskIcon */,
+			FTurnkeyEditorSupport::RunUAT(CommandLine, PlatformInfo->DisplayName, LOCTEXT("Turnkey_GettingCommandLine", "Copying Commandline"), LOCTEXT("Turnkey_CustomTaskNameCmdLine", "CustomCommandLine"), nullptr /* TaskIcon */,
 				[BuildName=Build.Name, ReportFilename](FString, double)
 				{
 					FString ReportContents;
@@ -1018,7 +1018,7 @@ static void FormatSdkInfo(const FString& PlatformOrDevice, const FTurnkeySdkInfo
 	}
 	else
 	{
-		OutToolTip = FText::Format(LOCTEXT("SdkInfo_ToolTip", "Information returned from:\nRunUAT Turnkey -command=VerifySdk -platform={1}"), Args);
+		OutToolTip = FText::Format(LOCTEXT("SdkInfo_ToolTipPlatform", "Information returned from:\nRunUAT Turnkey -command=VerifySdk -platform={1}"), Args);
 	}
 }
 
@@ -1305,16 +1305,16 @@ static void MakeTurnkeyPlatformMenu(UToolMenu* ToolMenu, FName IniPlatformName,
 						// these ToolTipFormats will be formatted inside the MakeCustomBuildMenuEntries function with some named args
 						// the empty strings passed in is the deviceId, which means to show builds that don't contain DeviceId components
 #if ALLOW_CONTROL_TO_COPY_COMMANDLINE
-						FText ToolTipFormat = LOCTEXT("TurnkeyTooltip_EngineCustomBuild_WithCopy", "Execute the '{BuildName}' command on {DeviceName}.\n{BuildHelp}\n--\nThis runs the following command:\nRunUAT {CommandLine}\nHold Control to copy the final underlying BuildCookRun commandline to the clipboard.");
+						FText ToolTipFormat = LOCTEXT("TurnkeyTooltip_EngineCustomBuild_WithCopyDevice", "Execute the '{BuildName}' command on {DeviceName}.\n{BuildHelp}\n--\nThis runs the following command:\nRunUAT {CommandLine}\nHold Control to copy the final underlying BuildCookRun commandline to the clipboard.");
 #else
-						FText ToolTipFormat = LOCTEXT("TurnkeyTooltip_EngineCustomBuild", "Execute the '{BuildName}' command on {DeviceName}.\n{BuildHelp}\n--\nThis runs the following command:\nRunUAT {CommandLine}");
+						FText ToolTipFormat = LOCTEXT("TurnkeyTooltip_EngineCustomBuildDevice", "Execute the '{BuildName}' command on {DeviceName}.\n{BuildHelp}\n--\nThis runs the following command:\nRunUAT {CommandLine}");
 #endif
 						MakeCustomBuildMenuEntries(CustomBuildSection, PackagingSettings->EngineCustomBuilds, IniPlatformName, DeviceId, DeviceName, ToolTipFormat);
 
 #if ALLOW_CONTROL_TO_COPY_COMMANDLINE
-						ToolTipFormat = LOCTEXT("TurnkeyTooltip_ProjectCustomBuild", "Execute this project's custom '{BuildName}' command on {DeviceName}.\n{BuildHelp}\n--\nThis runs the following command:\nRunUAT {CommandLine}\nHold Control to copy the final BuildCookRun commandline to the clipboard.\nThis custom command comes from Project Settings.");
+						ToolTipFormat = LOCTEXT("TurnkeyTooltip_ProjectCustomBuildDevice", "Execute this project's custom '{BuildName}' command on {DeviceName}.\n{BuildHelp}\n--\nThis runs the following command:\nRunUAT {CommandLine}\nHold Control to copy the final BuildCookRun commandline to the clipboard.\nThis custom command comes from Project Settings.");
 #else
-						ToolTipFormat = LOCTEXT("TurnkeyTooltip_ProjectCustomBuild_WithCopy", "Execute this project's custom '{BuildName}' command on {DeviceName}.\n{BuildHelp}\n--\nThis runs the following command:\nRunUAT {CommandLine}\nThis custom command comes from Project Settings.");
+						ToolTipFormat = LOCTEXT("TurnkeyTooltip_ProjectCustomBuild_WithCopyDevice", "Execute this project's custom '{BuildName}' command on {DeviceName}.\n{BuildHelp}\n--\nThis runs the following command:\nRunUAT {CommandLine}\nThis custom command comes from Project Settings.");
 #endif
 						MakeCustomBuildMenuEntries(CustomBuildSection, PackagingSettings->ProjectCustomBuilds, IniPlatformName, DeviceId, DeviceName, ToolTipFormat);
 					}
