Log badly formatted platform device IDs.
Fix incoming after we can see the faulty platform and device name.

diff --git a/Engine/Source/Editor/TurnkeySupport/Private/TurnkeyWrapper.cpp b/Engine/Source/Editor/TurnkeySupport/Private/TurnkeyWrapper.cpp
--- a/Engine/Source/Editor/TurnkeySupport/Private/TurnkeyWrapper.cpp
+++ b/Engine/Source/Editor/TurnkeySupport/Private/TurnkeyWrapper.cpp
@@ -1,6 +1,7 @@
 // Copyright Epic Games, Inc. All Rights Reserved.
 
 #include "TurnkeySupportModule.h"
+#include "TurnkeySupport.h"
 #include "Async/Async.h"
 #include "Misc/MonitoredProcess.h"
 
@@ -32,7 +33,11 @@ FString ConvertToUATPlatform(const FString& Platform)
 FString ConvertToUATDeviceId(const FString& DeviceId)
 {
 	TArray<FString> PlatformAndDevice;
-	DeviceId.ParseIntoArray(PlatformAndDevice, TEXT("@"), true);
+	int32 NumElems = DeviceId.ParseIntoArray(PlatformAndDevice, TEXT("@"), true);
+	if(NumElems < 2)
+	{
+		UE_LOG(LogTurnkeySupport, Fatal, TEXT("Badly formatted deviceId: %s"), *DeviceId);
+	}
 
 	return FString::Printf(TEXT("%s@%s"), *ConvertToUATPlatform(PlatformAndDevice[0]), *PlatformAndDevice[1]);
 }
@@ -40,7 +45,11 @@ FString ConvertToUATDeviceId(const FString& DeviceId)
 FString ConvertToDDPIDeviceId(const FString& DeviceId)
 {
 	TArray<FString> PlatformAndDevice;
-	DeviceId.ParseIntoArray(PlatformAndDevice, TEXT("@"), true);
+	int32 NumElems = DeviceId.ParseIntoArray(PlatformAndDevice, TEXT("@"), true);
+	if(NumElems < 2)
+	{
+		UE_LOG(LogTurnkeySupport, Fatal, TEXT("Badly formatted deviceId: %s"), *DeviceId);
+	}
 
 	return FString::Printf(TEXT("%s@%s"), *ConvertToDDPIPlatform(PlatformAndDevice[0]), *PlatformAndDevice[1]);
 }
