fix Quick Launch On incorrectly selecting Client build targets

#rnx

diff --git a/Engine/Source/Editor/TurnkeySupport/Private/TurnkeySupportModule.cpp b/Engine/Source/Editor/TurnkeySupport/Private/TurnkeySupportModule.cpp
--- a/Engine/Source/Editor/TurnkeySupport/Private/TurnkeySupportModule.cpp
+++ b/Engine/Source/Editor/TurnkeySupport/Private/TurnkeySupportModule.cpp
@@ -1569,6 +1569,7 @@ static void GenerateDeviceProxyMenuParams(TSharedPtr<ITargetDeviceProxy> DeviceP
 				if (DesktopPlatform->GetTargetsForCurrentProject().Num() > 0)
 				{
 					TArray<FTargetInfo> Targets = DesktopPlatform->GetTargetsForCurrentProject();
+					const FTargetInfo* TargetInfo = GetDefault<UProjectPackagingSettings>()->GetLaunchOnTargetInfo();
 
 					const TArray<FTargetInfo> ClientTargets = Targets.FilterByPredicate([](const FTargetInfo& Target)
 						{
@@ -1578,7 +1579,7 @@ static void GenerateDeviceProxyMenuParams(TSharedPtr<ITargetDeviceProxy> DeviceP
 					// we just want to know if any client build is selected, the correct flavor was picked above
 					for (const auto& Target : ClientTargets)
 					{
-						if (GetDefault<UProjectPackagingSettings>()->GetBuildTargetForPlatform(PlatformName) == Target.Name)
+						if (TargetInfo != nullptr && TargetInfo->Name == Target.Name)
 						{
 							VariantName += "Client";
 							break;
