early out of Turnkey sdk refresh if the engine is shutting down.

#rnx

diff --git a/Engine/Source/Editor/TurnkeySupport/Private/TurnkeySupportModule.cpp b/Engine/Source/Editor/TurnkeySupport/Private/TurnkeySupportModule.cpp
--- a/Engine/Source/Editor/TurnkeySupport/Private/TurnkeySupportModule.cpp
+++ b/Engine/Source/Editor/TurnkeySupport/Private/TurnkeySupportModule.cpp
@@ -2156,6 +2156,11 @@ void FTurnkeySupportModule::UpdateSdkInfo()
 
 		AsyncTask(ENamedThreads::GameThread, [this, ReportFilename, ExitCode, TurnkeyProcess, OnExitHandle]()
 		{
+			if (IsEngineExitRequested())
+			{
+				return;
+			}
+
 			FScopeLock Lock(&GTurnkeySection);
 
 			if (ExitCode == 0 || ExitCode == 10)
