Fixed textures failing to cook when running quick launching onto an Android device before having opened the turnkey android sub-menu.

#approval Jack Porter

diff --git a/Engine/Source/Editor/TurnkeySupport/Private/TurnkeySupportModule.cpp b/Engine/Source/Editor/TurnkeySupport/Private/TurnkeySupportModule.cpp
--- a/Engine/Source/Editor/TurnkeySupport/Private/TurnkeySupportModule.cpp
+++ b/Engine/Source/Editor/TurnkeySupport/Private/TurnkeySupportModule.cpp
@@ -1559,6 +1559,32 @@ static void GenerateDeviceProxyMenuParams(TSharedPtr<ITargetDeviceProxy> DeviceP
 				// only game and client devices are supported for launch on but devices only signal if they target client builds by their name e.g. "WindowsClient"
 				// if the user has a EBuildTargetType::Client target selected we will launch on a client device, otherwise we fall back to the default game device
 
+				// Initialize default flavor if not yet initialized
+				const PlatformInfo::FTargetPlatformInfo* VanillaInfo = PlatformInfo::FindVanillaPlatformInfo(PlatformName);
+				if(VanillaInfo != nullptr)
+				{
+					
+
+					// gather all valid flavors
+					const TArray<const PlatformInfo::FTargetPlatformInfo*> ValidFlavors = VanillaInfo->Flavors.FilterByPredicate([](const PlatformInfo::FTargetPlatformInfo* Target)
+						{
+							// Editor isn't a valid platform type that users can target
+							// The Build Target will choose client or server, so no need to show them as well
+							return Target->PlatformType != EBuildTargetType::Editor && Target->PlatformType != EBuildTargetType::Client && Target->PlatformType != EBuildTargetType::Server;
+						});
+
+					if (ValidFlavors.Num() > 1)
+					{
+						// Set the first flavor as the default if it hasn't been set
+						UProjectPackagingSettings* AllPlatformPackagingSettings = GetMutableDefault<UProjectPackagingSettings>();
+						FName CurrentFlavor = AllPlatformPackagingSettings->GetTargetFlavorForPlatform(PlatformName);
+						if (CurrentFlavor == VanillaInfo->Name)
+						{
+							FTurnkeySupportCallbacks::SetActiveFlavor(ValidFlavors[0]);
+						}
+					}
+				}
+
 				// We need to use flavors to launch on correctly on Android_ASTC / Android_ETC2, etc
 				const PlatformInfo::FTargetPlatformInfo* PlatformInfo = nullptr;
 				if (FApp::IsInstalled())
