WrangleContentForDebugging UAT command to prepare for debug iOS / tvOS

diff --git a/Engine/Source/Editor/TurnkeySupport/Private/TurnkeySupportModule.cpp b/Engine/Source/Editor/TurnkeySupport/Private/TurnkeySupportModule.cpp
--- a/Engine/Source/Editor/TurnkeySupport/Private/TurnkeySupportModule.cpp
+++ b/Engine/Source/Editor/TurnkeySupport/Private/TurnkeySupportModule.cpp
@@ -235,12 +235,6 @@ public:
 			return false;
 		}
 
-// 		// PrepForDebugging needs the platform to specify how
-// 		if (Mode == EPrepareContentMode::PrepareForDebugging)
-// 		{
-// 			return FDataDrivenPlatformInfoRegistry::GetPlatformInfo(IniPlatformName).PrepareForDebuggingOptions != TEXT("");
-// 		}
-
 		return true;
 	}
 
@@ -516,15 +510,18 @@ public:
 				BuildCookRunParams += TEXT(" -nodebuginfo");
 			}
 		}
-// 		else if (Mode == EPrepareContentMode::PrepareForDebugging)
-// 		{
-// 			if (ShouldBuildProject(PackagingSettings, TargetPlatform))
-// 			{
-// 				BuildCookRunParams += TEXT(" -build");
-// 			}
-// 
-// 			BuildCookRunParams += FString::Printf(TEXT(" %s"), *FDataDrivenPlatformInfoRegistry::GetPlatformInfo(IniPlatformName).PrepareForDebuggingOptions);
-// 		}
+        else if (Mode == EPrepareContentMode::PrepareForDebugging)
+         {
+             if (IniPlatformName == TEXT("IOS") || IniPlatformName == TEXT("TvOS"))
+             {
+                 FString CommandLine = FString::Printf(TEXT("WrangleContentForDebugging -platform=%s -ProjectFilePath=%s"), *IniPlatformName.ToString().ToLower(), *ProjectPath);
+              
+                 FTurnkeyEditorSupport::RunUAT(CommandLine, PlatformInfo->DisplayName, ContentPrepDescription, ContentPrepTaskName, ContentPrepIcon, &AnalyticsParamArray);
+             
+                 return;
+
+             }
+ 		}
 		else if (Mode == EPrepareContentMode::CookOnly)
 		{
 			ContentPrepDescription = LOCTEXT("CookingContentTaskName", "Cooking content");
@@ -1114,18 +1111,35 @@ static void MakeTurnkeyPlatformMenu(UToolMenu* ToolMenu, FName IniPlatformName,
 				FCanExecuteAction::CreateStatic(&FTurnkeySupportCallbacks::CanCookOrPackage, IniPlatformName, EPrepareContentMode::CookOnly)
 			)
 		);
-// 
-// 		Section.AddMenuEntry(
-//			NAME_None,
-// 			LOCTEXT("Turnkey_PrepareForDebugging", "Prepare For Debugging"),
-// 			LOCTEXT("TurnkeyTooltip_PrepareForDebugging", "Prepare this project for debugging"),
-// 			FSlateIcon(),
-// 			FUIAction(
-// 				FExecuteAction::CreateStatic(&FTurnkeySupportCallbacks::CookOrPackage, IniPlatformName, EPrepareContentMode::PrepareForDebugging),
-// 				FCanExecuteAction::CreateStatic(&FTurnkeySupportCallbacks::CanCookOrPackage, IniPlatformName, EPrepareContentMode::PrepareForDebugging)
-// 			)
-// 		);
-// 
+        // add platforms here when they have their own Prepare for debug flow
+#if PLATFORM_MAC
+        if (IniPlatformName.ToString() == "IOS")
+        {
+            Section.AddMenuEntry(
+                                 NAME_None,
+                                 LOCTEXT("Turnkey_PrepareForDebugging", "Prepare For Debugging"),
+                                 LOCTEXT("TurnkeyTooltip_PrepareForDebuggingIOS", "Prepare this project for debugging. Expects an IPA package with the same name as the project file in the Build/IOS/ folder."),
+                                 FSlateIcon(),
+                                 FUIAction(
+                                           FExecuteAction::CreateStatic(&FTurnkeySupportCallbacks::CookOrPackage, IniPlatformName, EPrepareContentMode::PrepareForDebugging),
+                                           FCanExecuteAction::CreateStatic(&FTurnkeySupportCallbacks::CanCookOrPackage, IniPlatformName, EPrepareContentMode::PrepareForDebugging)
+                                           )
+                                 );
+        }
+        if (IniPlatformName.ToString() == "TvOS")
+        {
+            Section.AddMenuEntry(
+                                 NAME_None,
+                                 LOCTEXT("Turnkey_PrepareForDebugging", "Prepare For Debugging"),
+                                 LOCTEXT("TurnkeyTooltip_PrepareForDebuggingIOS", "Prepare this project for debugging. Expects an IPA package with the same name as the project file in the Build/TVOS/ folder."),
+                                 FSlateIcon(),
+                                 FUIAction(
+                                           FExecuteAction::CreateStatic(&FTurnkeySupportCallbacks::CookOrPackage, IniPlatformName, EPrepareContentMode::PrepareForDebugging),
+                                           FCanExecuteAction::CreateStatic(&FTurnkeySupportCallbacks::CanCookOrPackage, IniPlatformName, EPrepareContentMode::PrepareForDebugging)
+                                           )
+                                 );
+        }
+#endif
 
 		UProjectPackagingSettings* PackagingSettings = FTurnkeySupportCallbacks::GetPackagingSettingsForPlatform(IniPlatformName);
 
diff --git a/Engine/Source/Programs/AutomationTool/AutomationUtils/Platform.cs b/Engine/Source/Programs/AutomationTool/AutomationUtils/Platform.cs
--- a/Engine/Source/Programs/AutomationTool/AutomationUtils/Platform.cs
+++ b/Engine/Source/Programs/AutomationTool/AutomationUtils/Platform.cs
@@ -863,6 +863,11 @@ namespace AutomationTool
 			return DirectoryReference.Combine(RuntimeRoot, RelativeProjectRootForStage.Name);
 		}
 
+		public virtual void PrepareForDebugging(string SourcePackage, string ProjectFilePath, string ClientPlatform)
+		{
+			LogError("Not implemented for the %s platform.", ClientPlatform);
+		}
+
 		// let the platform set the exe extension if it chooses (otherwise, use
 		// the switch statement in GetExeExtension below)
 		protected virtual string GetPlatformExeExtension()
diff --git a/Engine/Source/Programs/AutomationTool/IOS/IOSPlatform.Automation.cs b/Engine/Source/Programs/AutomationTool/IOS/IOSPlatform.Automation.cs
--- a/Engine/Source/Programs/AutomationTool/IOS/IOSPlatform.Automation.cs
+++ b/Engine/Source/Programs/AutomationTool/IOS/IOSPlatform.Automation.cs
@@ -2292,4 +2292,91 @@ public class IOSPlatform : ApplePlatform
 	{
 		return DirectoryReference.Combine(RuntimeRoot, "cookeddata/" + RelativeProjectRootForStage.Name);
 	}
+
+	public override void PrepareForDebugging(string SourcePackage, string ProjectFilePath, string ClientPlatform)
+	{
+		if (HostPlatform.Current.HostEditorPlatform != UnrealTargetPlatform.Mac)
+		{
+			LogInformation("Wrangling data for debug for an iOS/tvOS app for XCode is a Mac only feature. Aborting command.");
+			return;
+		}
+
+		int StartPos = ProjectFilePath.LastIndexOf("/");
+		int StringLength = ProjectFilePath.Length - 10; // 9 for .uproject, 1 for the /
+		string PackageName = ProjectFilePath.Substring(StartPos + 1, StringLength - StartPos);
+			if (string.IsNullOrEmpty(SourcePackage))
+		{
+			SourcePackage = ProjectFilePath;
+			SourcePackage = SourcePackage.Substring(0, SourcePackage.LastIndexOf('/'));
+			SourcePackage = SourcePackage + "/Build/" + ClientPlatform + '/' + PackageName + ".ipa";
+		}
+
+		string ZipFile = SourcePackage.Replace(".ipa", ".zip");
+
+		string PayloadPath = SourcePackage;
+		PayloadPath = PayloadPath.Substring(0, PayloadPath.LastIndexOf('/'));
+		PayloadPath += "/Payload/";
+		string CookedDataDirectory = PayloadPath + PackageName + ".app/cookeddata/";
+
+		LogInformation("ClientPlatform : {0}", ClientPlatform);
+		LogInformation("ProjectFilePath : {0}", ProjectFilePath);
+		LogInformation("Source : {0}", SourcePackage);
+		LogInformation("ZipFile {0}", ZipFile);
+		LogInformation("PackageName {0}", PackageName);
+		LogInformation("PayloadPath {0}", PayloadPath);
+
+		if (File.Exists(ZipFile))
+		{
+			LogInformation("Deleting previously present ZIP file created from IPA");
+			File.Delete(ZipFile);
+		}
+
+		File.Copy(SourcePackage, ZipFile);
+		UnzipPackage(ZipFile);
+
+		string ProjectPath = ProjectFilePath;
+		int Index = ProjectFilePath.IndexOf("?");
+		if (Index >= 0)
+		{
+		   ProjectPath = ProjectPath.Substring(0, Index);
+		}
+		if (Directory.Exists(ProjectPath + "/Binaries/" + ClientPlatform + '/' + "Payload/" + PackageName + ".app"))
+		{
+			CopyDirectory_NoExceptions(CookedDataDirectory, ProjectPath + "/Binaries/" + ClientPlatform + "/Payload/" + PackageName + ".app/cookeddata/", true);
+		}
+		else
+		{
+			string ProjectRoot = SourcePackage;
+			ProjectRoot = ProjectRoot.Substring(0, ProjectRoot.LastIndexOf('/'));
+			CopyFile(SourcePackage, ProjectRoot + "/Binaries/" + ClientPlatform + "/Payload/" + PackageName + ".ipa", true);
+		}
+
+		//cleanup
+		LogInformation("Deleting temp files ...");
+		File.Delete(ZipFile);
+		LogInformation("{0} deleted", ZipFile);
+		Directory.Delete(PayloadPath, true);
+		LogInformation("{0} deleted", PayloadPath);
+	}
+
+	public void UnzipPackage(string PackageToUnzip)
+	{
+		string UnzipPath = PackageToUnzip;
+		UnzipPath = UnzipPath.Substring(0, UnzipPath.LastIndexOf('/'));
+		LogInformation("Unzipping to {0}", UnzipPath);
+
+		using (Ionic.Zip.ZipFile Zip = new Ionic.Zip.ZipFile(PackageToUnzip))
+		{
+			foreach (Ionic.Zip.ZipEntry Entry in Zip.Entries.Where(x => !x.IsDirectory))
+			{
+				string OutputFileName = Path.Combine(UnzipPath, Entry.FileName);
+				Directory.CreateDirectory(Path.GetDirectoryName(OutputFileName));
+				using (FileStream OutputStream = new FileStream(OutputFileName, FileMode.Create, FileAccess.Write))
+				{
+					Entry.Extract(OutputStream);
+				}
+				LogInformation("Extracted {0}", OutputFileName);
+			}
+		}
+	}
 }
diff --git a/Engine/Source/Programs/AutomationTool/Scripts/WrangleContentForDebugging.Automation.cs b/Engine/Source/Programs/AutomationTool/Scripts/WrangleContentForDebugging.Automation.cs
new file mode 100644
--- /dev/null
+++ b/Engine/Source/Programs/AutomationTool/Scripts/WrangleContentForDebugging.Automation.cs
@@ -0,0 +1,39 @@
+// Copyright Epic Games, Inc. All Rights Reserved.
+
+using Microsoft.Extensions.Logging;
+using System.Collections.Generic;
+using AutomationTool;
+using UnrealBuildTool;
+using System.IO;
+//using System.Linq;
+using EpicGames.Core;
+
+namespace AutomationScripts
+{
+	public class WrangleContentForDebugging : BuildCommand
+	{
+		public override void ExecuteBuild()
+		{
+			LogInformation("======= WrangleContentForDebugging - Start =======");
+
+			if (ParseParam("usage"))
+			{
+				LogInformation("Arguments : \n -platform=<platform> \n -ProjectFilePath=<path/ProjectName.uproject> Mandatory. Path to a path to the roject you want to Wrangle Debug Data to. \n -SourcePackage=<path/PackageName.ipa> Optional. Specify the .ipa to wrangle data for debug from. When not specified, [ProjectPath]/Build/[IOS|TVOS]/[ProjectName].ipa will be used.");
+			}
+
+			FileReference ProjectFile = new FileReference(ParseParamValue("ProjectFilePath=", null));
+			//ProjectFile = ParseProjectParam();
+			var Params = new ProjectParams
+			(
+				Command: this,
+				RawProjectPath: ProjectFile
+			);
+			var SC = Project.CreateDeploymentContext(Params, false);
+
+			//Project newProject;
+			SC[0].StageTargetPlatform.PrepareForDebugging(ParseParamValue("SourcePackage=", null), ParseParamValue("ProjectFilePath=", null), ParseParamValue("platform=", null));
+
+			LogInformation("======= WrangleContentForDebugging - Done ========");
+		}
+	}
+}
\ No newline at end of file
