Fix localization warning

diff --git a/Engine/Source/Editor/TurnkeySupport/Private/TurnkeySupportModule.cpp b/Engine/Source/Editor/TurnkeySupport/Private/TurnkeySupportModule.cpp
--- a/Engine/Source/Editor/TurnkeySupport/Private/TurnkeySupportModule.cpp
+++ b/Engine/Source/Editor/TurnkeySupport/Private/TurnkeySupportModule.cpp
@@ -1131,7 +1131,7 @@ static void MakeTurnkeyPlatformMenu(UToolMenu* ToolMenu, FName IniPlatformName,
             Section.AddMenuEntry(
                                  NAME_None,
                                  LOCTEXT("Turnkey_PrepareForDebugging", "Prepare For Debugging"),
-                                 LOCTEXT("TurnkeyTooltip_PrepareForDebuggingIOS", "Prepare this project for debugging. Expects an IPA package with the same name as the project file in the Build/TVOS/ folder."),
+                                 LOCTEXT("TurnkeyTooltip_PrepareForDebuggingTvOS", "Prepare this project for debugging. Expects an IPA package with the same name as the project file in the Build/TVOS/ folder."),
                                  FSlateIcon(),
                                  FUIAction(
                                            FExecuteAction::CreateStatic(&FTurnkeySupportCallbacks::CookOrPackage, IniPlatformName, EPrepareContentMode::PrepareForDebugging),
