PR #9787: Add level editor extensibility hooks to TurnKey (Contributed by [REDACTED_PERSON])

diff --git a/Engine/Source/Editor/TurnkeySupport/Private/TurnkeyEditorSupport.cpp b/Engine/Source/Editor/TurnkeySupport/Private/TurnkeyEditorSupport.cpp
--- a/Engine/Source/Editor/TurnkeySupport/Private/TurnkeyEditorSupport.cpp
+++ b/Engine/Source/Editor/TurnkeySupport/Private/TurnkeyEditorSupport.cpp
@@ -52,7 +52,7 @@ FString FTurnkeyEditorSupport::GetUATOptions()
 void FTurnkeyEditorSupport::AddEditorOptions(FToolMenuSection& Section)
 {
 #if WITH_EDITOR
-	Section.AddSeparator(NAME_None);
+	Section.AddSeparator("TurnkeyOptionsSeparator");
 
 	Section.AddMenuEntry(
 		NAME_None,
diff --git a/Engine/Source/Editor/TurnkeySupport/Private/TurnkeySupportModule.cpp b/Engine/Source/Editor/TurnkeySupport/Private/TurnkeySupportModule.cpp
--- a/Engine/Source/Editor/TurnkeySupport/Private/TurnkeySupportModule.cpp
+++ b/Engine/Source/Editor/TurnkeySupport/Private/TurnkeySupportModule.cpp
@@ -48,6 +48,7 @@
 #include "ITurnkeyIOModule.h"
 #include "AnalyticsEventAttribute.h"
 #if WITH_EDITOR
+#include "LevelEditor.h"
 #include "Experimental/ZenServerInterface.h"
 #endif
 
@@ -1981,7 +1982,12 @@ TSharedRef<SWidget> FTurnkeySupportModule::MakeTurnkeyMenuWidget() const
 		}
 	}
 
+#if WITH_EDITOR
+	FLevelEditorModule& LevelEditorModule = FModuleManager::LoadModuleChecked<FLevelEditorModule>("LevelEditor");
+	FToolMenuContext MenuContext(FTurnkeySupportCommands::ActionList, LevelEditorModule.GetToolBarExtensibilityManager()->GetAllExtenders());
+#else
 	FToolMenuContext MenuContext(FTurnkeySupportCommands::ActionList);
+#endif
 	return UToolMenus::Get()->GenerateWidget(MenuName, MenuContext);
 
 }
