Move PerPlatformTargetFlavorName initialization to FTurnkeySupportModule::UpdateSdkInfo() so that automatically initialized following editor launch instead of being dependent on the Turnkey menu having been opened.

diff --git a/Engine/Source/Editor/TurnkeySupport/Private/TurnkeySupportModule.cpp b/Engine/Source/Editor/TurnkeySupport/Private/TurnkeySupportModule.cpp
--- a/Engine/Source/Editor/TurnkeySupport/Private/TurnkeySupportModule.cpp
+++ b/Engine/Source/Editor/TurnkeySupport/Private/TurnkeySupportModule.cpp
@@ -1165,6 +1165,7 @@ static void MakeTurnkeyPlatformMenu(UToolMenu* ToolMenu, FName IniPlatformName,
 
 		UProjectPackagingSettings* AllPlatformPackagingSettings = GetMutableDefault<UProjectPackagingSettings>();
 
+		// Populate Flavor Selection menu with available flavors
 		// Flavor Selection exists for Android to be able to package ASTC, DXT, ETC2, etc
 		{
 			// gather all valid flavors
@@ -1177,13 +1178,6 @@ static void MakeTurnkeyPlatformMenu(UToolMenu* ToolMenu, FName IniPlatformName,
 
 			if (ValidFlavors.Num() > 1)
 			{
-				// Set the first flavor as the default if it hasn't been set
-				FName CurrentFlavor = AllPlatformPackagingSettings->GetTargetFlavorForPlatform(IniPlatformName);
-				if (CurrentFlavor == VanillaInfo->Name)
-				{
-					FTurnkeySupportCallbacks::SetActiveFlavor(ValidFlavors[0]);
-				}
-
 				FToolMenuSection& FlavorSection = ToolMenu->AddSection("FlavorSelection", LOCTEXT("TurnkeySection_FlavorSelection", "Flavor Selection"));
 				
 				for (const PlatformInfo::FTargetPlatformInfo* Info : ValidFlavors)
@@ -1575,32 +1569,6 @@ static void GenerateDeviceProxyMenuParams(TSharedPtr<ITargetDeviceProxy> DeviceP
 				// only game and client devices are supported for launch on but devices only signal if they target client builds by their name e.g. "WindowsClient"
 				// if the user has a EBuildTargetType::Client target selected we will launch on a client device, otherwise we fall back to the default game device
 
-				// Initialize default flavor if not yet initialized
-				const PlatformInfo::FTargetPlatformInfo* VanillaInfo = PlatformInfo::FindVanillaPlatformInfo(PlatformName);
-				if(VanillaInfo != nullptr)
-				{
-					
-
-					// gather all valid flavors
-					const TArray<const PlatformInfo::FTargetPlatformInfo*> ValidFlavors = VanillaInfo->Flavors.FilterByPredicate([](const PlatformInfo::FTargetPlatformInfo* Target)
-						{
-							// Editor isn't a valid platform type that users can target
-							// The Build Target will choose client or server, so no need to show them as well
-							return Target->PlatformType != EBuildTargetType::Editor && Target->PlatformType != EBuildTargetType::Client && Target->PlatformType != EBuildTargetType::Server;
-						});
-
-					if (ValidFlavors.Num() > 1)
-					{
-						// Set the first flavor as the default if it hasn't been set
-						UProjectPackagingSettings* AllPlatformPackagingSettings = GetMutableDefault<UProjectPackagingSettings>();
-						FName CurrentFlavor = AllPlatformPackagingSettings->GetTargetFlavorForPlatform(PlatformName);
-						if (CurrentFlavor == VanillaInfo->Name)
-						{
-							FTurnkeySupportCallbacks::SetActiveFlavor(ValidFlavors[0]);
-						}
-					}
-				}
-
 				// We need to use flavors to launch on correctly on Android_ASTC / Android_ETC2, etc
 				const PlatformInfo::FTargetPlatformInfo* PlatformInfo = nullptr;
 				if (FApp::IsInstalled())
@@ -2254,6 +2222,33 @@ void FTurnkeySupportModule::UpdateSdkInfo()
 								PerPlatformSdkInfo[PlatformName].SDKVersions.Add(TEXT("AutoSDK"), SdkInfo.SDKVersions[TEXT("AutoSDK")]);
 							}
 						}
+
+						// initialize default flavor if not yet initialized
+						// flavor selection exists for Android to be able to package ASTC, DXT, ETC2, etc
+						const PlatformInfo::FTargetPlatformInfo* VanillaInfo = PlatformInfo::FindVanillaPlatformInfo(PlatformName);
+						if (VanillaInfo != nullptr)
+						{
+							// gather all valid flavors
+							const TArray<const PlatformInfo::FTargetPlatformInfo*> ValidFlavors = VanillaInfo->Flavors.FilterByPredicate([](const PlatformInfo::FTargetPlatformInfo* Target)
+								{
+									// Editor isn't a valid platform type that users can target
+									// The Build Target will choose client or server, so no need to show them as well
+									return Target->PlatformType != EBuildTargetType::Editor && Target->PlatformType != EBuildTargetType::Client && Target->PlatformType != EBuildTargetType::Server;
+								});
+
+							// if platform uses flavors and doesn't already have a default, set default flavor to the first availabl
+							if (ValidFlavors.Num() > 1)
+							{
+								UProjectPackagingSettings* AllPlatformPackagingSettings = GetMutableDefault<UProjectPackagingSettings>();
+								FName CurrentFlavor = AllPlatformPackagingSettings->GetTargetFlavorForPlatform(PlatformName);
+
+								if (CurrentFlavor == VanillaInfo->Name)
+								{
+									UE_LOG(LogTurnkeySupport, Log, TEXT("Platform %s uses flavors but has no default flavor set. Setting to %s"), *PlatformName.ToString(), *(ValidFlavors[0]->Name.ToString()));
+									FTurnkeySupportCallbacks::SetActiveFlavor(ValidFlavors[0]);
+								}
+							}
+						}
 					}
 				}
 
