Fix warning in turnkey

diff --git a/Engine/Source/Editor/TurnkeySupport/Private/TurnkeySupportModule.cpp b/Engine/Source/Editor/TurnkeySupport/Private/TurnkeySupportModule.cpp
--- a/Engine/Source/Editor/TurnkeySupport/Private/TurnkeySupportModule.cpp
+++ b/Engine/Source/Editor/TurnkeySupport/Private/TurnkeySupportModule.cpp
@@ -516,7 +516,8 @@ public:
 #if  PLATFORM_WINDOWS
              if (IniPlatformName == TEXT("IOS") || IniPlatformName == TEXT("TvOS"))
              {
-				 bool bSupportSecondaryMac = GConfig->GetBool(TEXT("/Script/IOSRuntimeSettings.IOSRuntimeSettings"), TEXT("bSupportSecondaryMac"), bSupportSecondaryMac, GEngineIni);
+				 bool bSupportSecondaryMac = false;
+				 bSupportSecondaryMac = GConfig->GetBool(TEXT("/Script/IOSRuntimeSettings.IOSRuntimeSettings"), TEXT("bSupportSecondaryMac"), bSupportSecondaryMac, GEngineIni);
 				 if (bSupportSecondaryMac)
 				 {
 					 FString OtherCommandLine = FString::Printf(TEXT("SetSecondaryRemoteMac -platform=%s -ProjectFilePath=%s"), *IniPlatformName.ToString().ToLower(), *ProjectPath);
