- Disabled some Turnkey async stuff is GIsRequestingExit, if an early shutdown happens (like when choosing a project to open)

diff --git a/Engine/Source/Editor/TurnkeySupport/Private/TurnkeySupportModule.cpp b/Engine/Source/Editor/TurnkeySupport/Private/TurnkeySupportModule.cpp
--- a/Engine/Source/Editor/TurnkeySupport/Private/TurnkeySupportModule.cpp
+++ b/Engine/Source/Editor/TurnkeySupport/Private/TurnkeySupportModule.cpp
@@ -2407,7 +2407,7 @@ void FTurnkeySupportModule::UpdateSdkInfoForDevices(TArray<FString> PlatformDevi
 		{
 			FScopeLock Lock(&GTurnkeySection);
 
-			if (ExitCode == 0 || ExitCode == 10)
+			if (!IsEngineExitRequested() && (ExitCode == 0 || ExitCode == 10))
 			{
 				TArray<FString> Contents;
 				if (FFileHelper::LoadFileToStringArray(Contents, *ReportFilename))
