Turnkey:  Turnkey and editor support for controlling auto software update modes on platform devices

- this change just adds support to Turnkey and the editor UI to expose this support for supported platforms
  - follow on changes will be required to support the feature on a given platform

- Turnkey now tracks the auto software update mode per device
  - unknown/enable/disable
    - will be unknown for devices which do not support auto software updates or platforms which have not yet implemented support
- new UAT platform method SetDeviceAutoSoftwareUpdateMode() must be implemented for supporting platforms
- new Turnkey command DeviceAutoSoftwareUpdates for changing auto software update state on a device
- edtior exposes this functionality under the device entries in the platforms menu
  - devices which support auto software updates will expose an "Enable auto software updates" check box item in their Sdk Info menu
    - toggling this checkbox changes the auto software update policy on the device

diff --git a/Engine/Source/Editor/TurnkeySupport/Private/TurnkeySupportModule.cpp b/Engine/Source/Editor/TurnkeySupport/Private/TurnkeySupportModule.cpp
--- a/Engine/Source/Editor/TurnkeySupport/Private/TurnkeySupportModule.cpp
+++ b/Engine/Source/Editor/TurnkeySupport/Private/TurnkeySupportModule.cpp
@@ -1097,6 +1097,26 @@ static void MakeCustomBuildMenuEntries(FToolMenuSection& Section, const TArray<F
 
 }
 
+
+static void TurnkeySetDeviceAutoSoftwareUpdates(FString PlatformName, FString DeviceId, bool bEnableAutoSoftwareUpdates)
+{
+	FString CommandLine;
+	const FString ProjectPath = GetProjectPathForTurnkey();
+	if (!ProjectPath.IsEmpty())
+	{
+		CommandLine.Appendf(TEXT(" -ScriptsForProject=\"%s\" "), *ProjectPath);
+	}
+	CommandLine.Appendf(TEXT("Turnkey -command=DeviceAutoSoftwareUpdates -device=%s -enable=%s"), *DeviceId, bEnableAutoSoftwareUpdates ? TEXT("true") : TEXT("false"));
+
+	FText TaskName = LOCTEXT("SetDeviceAutoSoftwareUpdates", "Set device auto software updates");
+	FTurnkeyEditorSupport::RunUAT(CommandLine, FText::FromString(PlatformName), TaskName, TaskName, FAppStyle::GetBrush(TEXT("MainFrame.PackageProject")), nullptr,
+		[DeviceId](FString, double)
+		{
+			ITurnkeySupportModule::Get().UpdateSdkInfoForDevices(TArray<FString>{DeviceId});
+		});	
+}
+
+
 static void MakeTurnkeyPlatformMenu(UToolMenu* ToolMenu, FName IniPlatformName, ITargetDeviceServicesModule* TargetDeviceServicesModule)
 {
 	const FDataDrivenPlatformInfo& DDPI = FDataDrivenPlatformInfoRegistry::GetPlatformInfo(IniPlatformName);
@@ -1361,6 +1381,27 @@ static void MakeTurnkeyPlatformMenu(UToolMenu* ToolMenu, FName IniPlatformName,
 						)
 					);
 
+					ETurnkeyDeviceAutoSoftwareUpdateMode DeviceAutoUpdateStatus = SdkInfo.DeviceAutoSoftwareUpdates;
+					if (DeviceAutoUpdateStatus != ETurnkeyDeviceAutoSoftwareUpdateMode::Unknown)
+					{
+
+						Section.AddMenuEntry(
+							NAME_None,
+							LOCTEXT("Turnkey_EnableAutoSoftwareUpdates", "Enable auto software updates"),
+							LOCTEXT("Turnkey_EnableAutoSoftwareUpdatesDescription", "This device will automatically update its software when new versions are released"),
+							FSlateIcon(),
+							FUIAction(
+								FExecuteAction::CreateStatic(TurnkeySetDeviceAutoSoftwareUpdates, IniPlatformName.ToString(), DeviceId, DeviceAutoUpdateStatus != ETurnkeyDeviceAutoSoftwareUpdateMode::Enabled),
+								FCanExecuteAction(),
+								FIsActionChecked::CreateLambda([DeviceAutoUpdateStatus]()
+						{
+							return DeviceAutoUpdateStatus == ETurnkeyDeviceAutoSoftwareUpdateMode::Enabled;
+						})
+							),
+							EUserInterfaceActionType::ToggleButton
+							);
+					}
+
 					if (SdkInfo.DeviceStatus == ETurnkeyDeviceStatus::SoftwareValid)
 					{
 						Section.AddMenuEntry(
@@ -2118,6 +2159,17 @@ bool GetSdkInfoFromTurnkey(FString Line, FName& PlatformName, FString& DeviceId,
 		SdkInfo.DeviceStatus = ETurnkeyDeviceStatus::SoftwareInvalid;
 	}
 
+	// Device auto update support
+	SdkInfo.DeviceAutoSoftwareUpdates = ETurnkeyDeviceAutoSoftwareUpdateMode::Unknown;
+	if (FlagsString.Contains(TEXT("Device_AutoSoftwareUpdates_Disabled")))
+	{
+		SdkInfo.DeviceAutoSoftwareUpdates = ETurnkeyDeviceAutoSoftwareUpdateMode::Disabled;
+	}
+	else if (FlagsString.Contains(TEXT("Device_AutoSoftwareUpdates_Enabled")))
+	{
+		SdkInfo.DeviceAutoSoftwareUpdates = ETurnkeyDeviceAutoSoftwareUpdateMode::Enabled;
+	}
+
 	return true;
 }
 
diff --git a/Engine/Source/Editor/TurnkeySupport/Public/Interfaces/ITurnkeySupportModule.h b/Engine/Source/Editor/TurnkeySupport/Public/Interfaces/ITurnkeySupportModule.h
--- a/Engine/Source/Editor/TurnkeySupport/Public/Interfaces/ITurnkeySupportModule.h
+++ b/Engine/Source/Editor/TurnkeySupport/Public/Interfaces/ITurnkeySupportModule.h
@@ -28,11 +28,19 @@ enum class ETurnkeyDeviceStatus : uint8
 	SoftwareInvalid,
 };
 
+enum class ETurnkeyDeviceAutoSoftwareUpdateMode : uint8
+{
+	Unknown,
+	Disabled,
+	Enabled,
+};
+
 struct FTurnkeySdkInfo
 {
 	struct Version { FString Min, Max, Current; };
 	ETurnkeyPlatformSdkStatus Status = ETurnkeyPlatformSdkStatus::Unknown;
 	ETurnkeyDeviceStatus DeviceStatus = ETurnkeyDeviceStatus::Unknown;
+	ETurnkeyDeviceAutoSoftwareUpdateMode DeviceAutoSoftwareUpdates = ETurnkeyDeviceAutoSoftwareUpdateMode::Unknown;
 	FText SdkErrorInformation;
 	TMap<FString, Version> SDKVersions;
 	bool bCanInstallFullSdk;
diff --git a/Engine/Source/Programs/AutomationTool/AutomationUtils/Platform.cs b/Engine/Source/Programs/AutomationTool/AutomationUtils/Platform.cs
--- a/Engine/Source/Programs/AutomationTool/AutomationUtils/Platform.cs
+++ b/Engine/Source/Programs/AutomationTool/AutomationUtils/Platform.cs
@@ -31,12 +31,19 @@ namespace AutomationTool
 
 	public class DeviceInfo
 	{
+		public enum AutoSoftwareUpdateMode
+		{
+			Unknown,
+			Disabled,
+			Enabled
+		}
+
 		public DeviceInfo(UnrealTargetPlatform Platform)
 		{
 			this.Platform = Platform;
 		}
 
-		public DeviceInfo(UnrealTargetPlatform Platform, string Name, string Id, string SoftwareVersion, string Type, bool bIsDefault, bool bCanConnect, Dictionary<string, string> PlatformValues=null)
+		public DeviceInfo(UnrealTargetPlatform Platform, string Name, string Id, string SoftwareVersion, string Type, bool bIsDefault, bool bCanConnect, Dictionary<string, string> PlatformValues = null, AutoSoftwareUpdateMode AutoSoftwareUpdates = AutoSoftwareUpdateMode.Unknown)
 		{
 			this.Platform = Platform;
 			this.Name = Name;
@@ -45,6 +52,7 @@ namespace AutomationTool
 			this.Type = Type;
 			this.bIsDefault = bIsDefault;
 			this.bCanConnect = bCanConnect;
+			this.AutoSoftwareUpdates = AutoSoftwareUpdates;
 			if (PlatformValues != null)
 			{
 				this.PlatformValues = new Dictionary<string, string>(PlatformValues);
@@ -58,10 +66,11 @@ namespace AutomationTool
 		public string Type;
 		public bool bIsDefault = false;
 		// is the device able to be connected to (this is more about able to flash SDK or run, not about matching SDK version)
-		// if false, any of the above fields are suspect, especually SoftwareVersion
+		// if false, any of the above fields are suspect, especially SoftwareVersion
 		public bool bCanConnect = true;
+		public AutoSoftwareUpdateMode AutoSoftwareUpdates = AutoSoftwareUpdateMode.Unknown;
 
-		// case insesitive platform value dictionary. turnkey doesn't use this, but the platform can look up the device during deployment, etc to get this out
+		// case insensitive platform value dictionary. turnkey doesn't use this, but the platform can look up the device during deployment, etc to get this out
 		public Dictionary<string, string> PlatformValues = new Dictionary<string, string>(StringComparer.InvariantCultureIgnoreCase);
 	}
 
@@ -268,10 +277,12 @@ namespace AutomationTool
 			Params = null;
 			return false;
 		}
+
 		public virtual bool GetSDKInstallCommand(out string Command, out string Params, ref bool bRequiresPrivilegeElevation, ref bool bCreateWindow, ITurnkeyContext TurnkeyContext, bool bSdkAlreadyInstalled)
 		{
 			return GetSDKInstallCommand(out Command, out Params, ref bRequiresPrivilegeElevation, ref bCreateWindow, TurnkeyContext);
 		}
+
 		public virtual bool GetDeviceUpdateSoftwareCommand(out string Command, out string Params, ref bool bRequiresPrivilegeElevation, ref bool bCreateWindow, ITurnkeyContext TurnkeyContext, DeviceInfo Device = null)
 		{
 			Command = null;
@@ -305,6 +316,12 @@ namespace AutomationTool
 			return true;
 		}
 
+		public virtual bool SetDeviceAutoSoftwareUpdateMode(DeviceInfo Device, bool bEnableAutoSoftwareUpdates)
+		{
+			LogWarning("{0} does not implement SetDeviceAutoSoftwareUpdateMode", PlatformType);
+			return false;
+		}
+
 		#endregion
 
 		/// <summary>
diff --git a/Engine/Source/Programs/AutomationTool/Turnkey/Commands/DeviceAutoSoftwareUpdates.cs b/Engine/Source/Programs/AutomationTool/Turnkey/Commands/DeviceAutoSoftwareUpdates.cs
new file mode 100644
--- /dev/null
+++ b/Engine/Source/Programs/AutomationTool/Turnkey/Commands/DeviceAutoSoftwareUpdates.cs
@@ -0,0 +1,51 @@
+// Copyright Epic Games, Inc. All Rights Reserved.
+
+using System;
+using System.Collections.Generic;
+using AutomationTool;
+using UnrealBuildTool;
+
+namespace Turnkey.Commands
+{
+	class DeviceAutoSoftwareUpdates : TurnkeyCommand
+	{
+		protected override CommandGroup Group => CommandGroup.Sdk;
+
+		protected override void Execute(string[] CommandOptions)
+		{
+			// Are we enabling or disabling auto updates on the devices?
+			bool bUnattended = TurnkeyUtils.ParseParam("Unattended", CommandOptions);
+			string EnableArg = TurnkeyUtils.ParseParamValue("Enable", null, CommandOptions);
+
+			bool bEnableAutoSoftwareUpdates = false;
+			if (string.IsNullOrEmpty(EnableArg))
+			{
+				if (bUnattended)
+				{
+					TurnkeyUtils.Log("Error: Desired state of device auto software updates must be specified via -Enable=true/false");
+					return;
+				}
+				else
+				{
+					bEnableAutoSoftwareUpdates = TurnkeyUtils.ReadInputInt("Enable device auto software updates?", new List<string> { "Disable", "Enable" }, false) != 0;
+				}
+			}
+			else
+			{
+				bEnableAutoSoftwareUpdates = EnableArg.Equals("true", StringComparison.InvariantCultureIgnoreCase);
+			}
+
+			// Get list of devices
+			List<UnrealTargetPlatform> PlatformsWithSdks = TurnkeyManifest.GetPlatformsWithSdks();
+			List<DeviceInfo> ChosenDevices = TurnkeyUtils.GetDevicesFromCommandLineOrUser(CommandOptions, PlatformsWithSdks);
+
+
+			// Set the auto updates mode on each device
+			foreach (DeviceInfo Device in ChosenDevices)
+			{
+				AutomationTool.Platform AutomationPlatform = Platform.GetPlatform(Device.Platform);
+				AutomationPlatform.SetDeviceAutoSoftwareUpdateMode(Device, bEnableAutoSoftwareUpdates);
+			}
+		}
+	}
+}
diff --git a/Engine/Source/Programs/AutomationTool/Turnkey/Commands/VerifySdk.cs b/Engine/Source/Programs/AutomationTool/Turnkey/Commands/VerifySdk.cs
--- a/Engine/Source/Programs/AutomationTool/Turnkey/Commands/VerifySdk.cs
+++ b/Engine/Source/Programs/AutomationTool/Turnkey/Commands/VerifySdk.cs
@@ -228,7 +228,7 @@ namespace Turnkey.Commands
 					}
 				}
 
-				// now check software verison of each device
+				// now check software version of each device
 				if (DevicesToCheck != null)
 				{
 					foreach (DeviceInfo Device in DevicesToCheck.Where(x => x.Platform == Platform))
@@ -263,6 +263,14 @@ namespace Turnkey.Commands
 							DeviceState |= SdkUtils.LocalAvailability.Device_CannotConnect;
 						}
 
+						if (Device.AutoSoftwareUpdates == AutomationTool.DeviceInfo.AutoSoftwareUpdateMode.Disabled)
+						{
+							DeviceState |= SdkUtils.LocalAvailability.Device_AutoSoftwareUpdates_Disabled;
+						}
+						else if (Device.AutoSoftwareUpdates == AutomationTool.DeviceInfo.AutoSoftwareUpdateMode.Enabled)
+						{
+							DeviceState |= SdkUtils.LocalAvailability.Device_AutoSoftwareUpdates_Enabled;
+						}
 
 						List<string> DeviceProperties = new List<string>()
 							{
diff --git a/Engine/Source/Programs/AutomationTool/Turnkey/SdkInfo.cs b/Engine/Source/Programs/AutomationTool/Turnkey/SdkInfo.cs
--- a/Engine/Source/Programs/AutomationTool/Turnkey/SdkInfo.cs
+++ b/Engine/Source/Programs/AutomationTool/Turnkey/SdkInfo.cs
@@ -18,26 +18,29 @@ namespace Turnkey
 		public enum LocalAvailability
 		{
 			None = 0,
-			AutoSdk_VariableExists                  = (1 << 0),
-			AutoSdk_ValidVersionExists              = (1 << 1),
-			AutoSdk_InvalidVersionExists            = (1 << 2),
+			AutoSdk_VariableExists					= (1 << 0),
+			AutoSdk_ValidVersionExists				= (1 << 1),
+			AutoSdk_InvalidVersionExists			= (1 << 2),
 			
-			InstalledSdk_ValidInactiveVersionExists = (1 << 3),
-			InstalledSdk_ValidVersionExists         = (1 << 4),
-			InstalledSdk_InvalidVersionExists       = (1 << 5),
+			InstalledSdk_ValidInactiveVersionExists	= (1 << 3),
+			InstalledSdk_ValidVersionExists			= (1 << 4),
+			InstalledSdk_InvalidVersionExists		= (1 << 5),
 
-			Platform_ValidHostPrerequisites         = (1 << 6),
-			Platform_InvalidHostPrerequisites       = (1 << 7),
+			Platform_ValidHostPrerequisites			= (1 << 6),
+			Platform_InvalidHostPrerequisites		= (1 << 7),
 
-			Device_InvalidPrerequisites             = (1 << 8),
-			Device_InstallSoftwareValid             = (1 << 9),
-			Device_InstallSoftwareInvalid           = (1 << 10),
-			Device_CannotConnect                    = (1 << 11),
+			Device_InvalidPrerequisites				= (1 << 8),
+			Device_InstallSoftwareValid				= (1 << 9),
+			Device_InstallSoftwareInvalid			= (1 << 10),
+			Device_CannotConnect					= (1 << 11),
 
-			Support_FullSdk                         = (1 << 12),
-			Support_AutoSdk                         = (1 << 13),
+			Support_FullSdk							= (1 << 12),
+			Support_AutoSdk							= (1 << 13),
 
-			Sdk_HasBestVersion                      = (1 << 14),
+			Sdk_HasBestVersion						= (1 << 14),
+
+			Device_AutoSoftwareUpdates_Disabled		= (1 << 15),
+			Device_AutoSoftwareUpdates_Enabled		= (1 << 16),
 		}
 
 		static public LocalAvailability GetLocalAvailability(AutomationTool.Platform AutomationPlatform, bool bAllowUpdatingPrerequisites, TurnkeyContextImpl TurnkeyContext)
