crash fix when receiving unknown device from Turnkey.

#rnx

diff --git a/Engine/Source/Editor/TurnkeySupport/Private/TurnkeySupportModule.cpp b/Engine/Source/Editor/TurnkeySupport/Private/TurnkeySupportModule.cpp
--- a/Engine/Source/Editor/TurnkeySupport/Private/TurnkeySupportModule.cpp
+++ b/Engine/Source/Editor/TurnkeySupport/Private/TurnkeySupportModule.cpp
@@ -2485,6 +2485,7 @@ void FTurnkeySupportModule::UpdateSdkInfoForDevices(TArray<FString> PlatformDevi
 						if (!PerDeviceSdkInfo.Contains(DDPIDeviceId))
 						{
 							UE_LOG(LogTurnkeySupport, Error, TEXT("Received DeviceId %s from Turnkey, but the engine doesn't know about it."), *DDPIDeviceId);
+							continue;
 						}
 
 						UE_LOG(LogTurnkeySupport, Log, TEXT("Turnkey Device: %s"), *Line);
