- Moved per-user-facing in ProjectPackagingSettings into a new class, UPlaformsMenuSettings, so that we don't sometimes write to DefaultGame.ini and sometimes Saved/Game.ini for the same class, which can cause issues where the editor has one value that doesn't match a target platform
- StagingDirectory, which was the user's location to store archived builds, not the actual StagingDirectory, is no longer exposed in settings

diff --git a/Engine/Config/BaseGame.ini b/Engine/Config/BaseGame.ini
--- a/Engine/Config/BaseGame.ini
+++ b/Engine/Config/BaseGame.ini
@@ -151,6 +151,13 @@ WriteBackMetadataToAssetRegistry=Disabled
 +IniSectionDenylist=HordeStorageServers
 +IniSectionDenylist=StorageServers
 
+[/Script/DeveloperToolSettings.PlatformsMenuSettings]
+StagingDirectory=(Path="")
+LaunchOnTarget=
+PerPlatformBuildConfig=()
+PerPlatformTargetFlavorName=(("Android", "Android_ASTC"))
+PerPlatformBuildTarget=()
+
 [/Script/Engine.HUD]
 DebugDisplay=AI
 
diff --git a/Engine/Source/Developer/DeveloperToolSettings/Classes/Settings/PlatformsMenuSettings.h b/Engine/Source/Developer/DeveloperToolSettings/Classes/Settings/PlatformsMenuSettings.h
new file mode 100644
--- /dev/null
+++ b/Engine/Source/Developer/DeveloperToolSettings/Classes/Settings/PlatformsMenuSettings.h
@@ -0,0 +1,60 @@
+// Copyright Epic Games, Inc. All Rights Reserved.
+
+#pragma once
+
+#include "CoreMinimal.h"
+#include "UObject/ObjectMacros.h"
+#include "UObject/Object.h"
+#include "ProjectPackagingSettings.h"
+#include "PlatformsMenuSettings.generated.h"
+
+
+
+UCLASS(config=Game)
+class DEVELOPERTOOLSETTINGS_API UPlatformsMenuSettings
+	: public UObject
+{
+	GENERATED_UCLASS_BODY()
+	
+	/** The directory to which the packaged project will be copied. */
+	UPROPERTY(config, EditAnywhere, Category=Project)
+	FDirectoryPath StagingDirectory;
+
+	/** Name of the target to use for LaunchOn (only Game/Client targets) */
+	UPROPERTY(config)
+	FString LaunchOnTarget;
+
+	/** Gets the current launch on target, checking that it's valid, and the default build target if it is not */
+	const FTargetInfo* GetLaunchOnTargetInfo() const;
+	
+	/**
+	 * Get and set the per-platform build config and targetplatform settings for the Turnkey/Launch on menu
+	 */
+	EProjectPackagingBuildConfigurations GetBuildConfigurationForPlatform(FName PlatformName) const;
+	void SetBuildConfigurationForPlatform(FName PlatformName, EProjectPackagingBuildConfigurations Configuration);
+
+	FName GetTargetFlavorForPlatform(FName PlatformName) const;
+	void SetTargetFlavorForPlatform(FName PlatformName, FName TargetFlavorName);
+
+	FString GetBuildTargetForPlatform(FName PlatformName) const;
+	void SetBuildTargetForPlatform(FName PlatformName, FString BuildTargetName);
+
+	const FTargetInfo* GetBuildTargetInfoForPlatform(FName PlatformName, bool& bOutIsProjectTarget) const;
+
+private:
+	/** Per platform build configuration */
+	UPROPERTY(config)
+	TMap<FName, EProjectPackagingBuildConfigurations> PerPlatformBuildConfig;
+
+	/** Per platform flavor cooking target */
+	UPROPERTY(config)
+	TMap<FName, FName> PerPlatformTargetFlavorName;
+
+	/** Per platform build target */
+	UPROPERTY(config, EditAnywhere, Category=Project)
+	TMap<FName, FString> PerPlatformBuildTarget;
+
+
+	
+};
+
diff --git a/Engine/Source/Developer/DeveloperToolSettings/Classes/Settings/ProjectPackagingSettings.h b/Engine/Source/Developer/DeveloperToolSettings/Classes/Settings/ProjectPackagingSettings.h
--- a/Engine/Source/Developer/DeveloperToolSettings/Classes/Settings/ProjectPackagingSettings.h
+++ b/Engine/Source/Developer/DeveloperToolSettings/Classes/Settings/ProjectPackagingSettings.h
@@ -217,14 +217,6 @@ public:
 	UPROPERTY(config, EditAnywhere, Category=Project)
 	FString BuildTarget;
 
-	/** Name of the target to use for LaunchOn (only Game/Client targets) */
-	UPROPERTY(config)
-	FString LaunchOnTarget;
-
-	/** The directory to which the packaged project will be copied. */
-	UPROPERTY(config, EditAnywhere, Category=Project)
-	FDirectoryPath StagingDirectory;
-
 	/**
 	 * If enabled, a full rebuild will be enforced each time the project is being packaged.
 	 * If disabled, only modified files will be built, which can improve iteration time.
@@ -638,34 +630,7 @@ public:
 	UPROPERTY(config)
 	TArray<FProjectBuildSettings> EngineCustomBuilds;
 
-
-	/**
-	 * Get and set the per-platform build config and targetplatform settings for the Turnkey/Launch on menu
-	 */
-	EProjectPackagingBuildConfigurations GetBuildConfigurationForPlatform(FName PlatformName) const;
-	void SetBuildConfigurationForPlatform(FName PlatformName, EProjectPackagingBuildConfigurations Configuration);
-
-	FName GetTargetFlavorForPlatform(FName PlatformName) const;
-	void SetTargetFlavorForPlatform(FName PlatformName, FName TargetFlavorName);
-
-	FString GetBuildTargetForPlatform(FName PlatformName) const;
-	void SetBuildTargetForPlatform(FName PlatformName, FString BuildTargetName);
-
-	const FTargetInfo* GetBuildTargetInfoForPlatform(FName PlatformName, bool& bOutIsProjectTarget) const;
-
 private:
-	/** Per platform build configuration */
-	UPROPERTY(config)
-	TMap<FName, EProjectPackagingBuildConfigurations> PerPlatformBuildConfig;
-
-	/** Per platform flavor cooking target */
-	UPROPERTY(config)
-	TMap<FName, FName> PerPlatformTargetFlavorName;
-
-	/** Per platform build target */
-	UPROPERTY(config, EditAnywhere, Category=Project)
-	TMap<FName, FString> PerPlatformBuildTarget;
-
 	/** Helper array used to mirror Blueprint asset selections across edits */
 	TArray<FFilePath> CachedNativizeBlueprintAssets;
 	
@@ -700,9 +665,6 @@ public:
 	/** Gets the current build target, checking that it's valid, and the default build target if it is not */
 	const FTargetInfo* GetBuildTargetInfo() const;
 
-	/** Gets the current launch ojn target, checking that it's valid, and the default build target if it is not */
-	const FTargetInfo* GetLaunchOnTargetInfo() const;
-
 	/** For non-default object instances, this will LoadConfig for a specific platform,*/
 	void LoadSettingsForPlatform(FString PlatformName)
 	{
diff --git a/Engine/Source/Developer/DeveloperToolSettings/Private/PlatformsMenuSettings.cpp b/Engine/Source/Developer/DeveloperToolSettings/Private/PlatformsMenuSettings.cpp
new file mode 100644
--- /dev/null
+++ b/Engine/Source/Developer/DeveloperToolSettings/Private/PlatformsMenuSettings.cpp
@@ -0,0 +1,88 @@
+// Copyright Epic Games, Inc. All Rights Reserved.
+
+#include "Settings/PlatformsMenuSettings.h"
+#include "Misc/App.h"
+#include "Misc/ConfigCacheIni.h"
+#include "Interfaces/IProjectManager.h"
+#include "DesktopPlatformModule.h"
+#include "DeveloperToolSettingsDelegates.h"
+#include "InstalledPlatformInfo.h"
+
+#include UE_INLINE_GENERATED_CPP_BY_NAME(PlatformsMenuSettings)
+
+#define LOCTEXT_NAMESPACE "SettingsClasses"
+
+extern const FTargetInfo* FindBestTargetInfo(const FString& TargetName, bool bContentOnlyUsesEngineTargets, bool* bOutIsProjectTarget);
+
+
+UPlatformsMenuSettings::UPlatformsMenuSettings(const FObjectInitializer& ObjectInitializer)
+	: Super(ObjectInitializer)
+{
+}
+
+
+const FTargetInfo* UPlatformsMenuSettings::GetBuildTargetInfoForPlatform(FName PlatformName, bool& bOutIsProjectTarget) const
+{
+	return FindBestTargetInfo(GetBuildTargetForPlatform(PlatformName), true, &bOutIsProjectTarget);
+}
+
+EProjectPackagingBuildConfigurations UPlatformsMenuSettings::GetBuildConfigurationForPlatform(FName PlatformName) const
+{
+	const EProjectPackagingBuildConfigurations* Value = PerPlatformBuildConfig.Find(PlatformName);
+
+	// PPBC_MAX defines the default project setting case and should be handled accordingly.
+	return Value == nullptr ? EProjectPackagingBuildConfigurations::PPBC_MAX : *Value;
+}
+
+void UPlatformsMenuSettings::SetBuildConfigurationForPlatform(FName PlatformName, EProjectPackagingBuildConfigurations Configuration)
+{
+	if (Configuration == EProjectPackagingBuildConfigurations::PPBC_MAX)
+	{
+		PerPlatformBuildConfig.Remove(PlatformName);
+	}
+	else
+	{
+		PerPlatformBuildConfig.Add(PlatformName, Configuration);
+	}
+}
+
+FName UPlatformsMenuSettings::GetTargetFlavorForPlatform(FName FlavorName) const
+{
+	const FName* Value = PerPlatformTargetFlavorName.Find(FlavorName);
+
+	// the flavor name is also the name of the vanilla info
+	return Value == nullptr ? FlavorName : *Value;
+}
+
+void UPlatformsMenuSettings::SetTargetFlavorForPlatform(FName PlatformName, FName TargetFlavorName)
+{
+	PerPlatformTargetFlavorName.Add(PlatformName, TargetFlavorName);
+}
+
+FString UPlatformsMenuSettings::GetBuildTargetForPlatform(FName PlatformName) const
+{
+	const FString* Value = PerPlatformBuildTarget.Find(PlatformName);
+
+	// empty string defines the default project setting case and should be handled accordingly.
+	return Value == nullptr ? "" : *Value;
+}
+
+void UPlatformsMenuSettings::SetBuildTargetForPlatform(FName PlatformName, FString BuildTargetName)
+{
+	if (BuildTargetName.IsEmpty())
+	{
+		PerPlatformBuildTarget.Remove(PlatformName);
+	}
+	else
+	{
+		PerPlatformBuildTarget.Add(PlatformName, BuildTargetName);
+	}
+}
+
+const FTargetInfo* UPlatformsMenuSettings::GetLaunchOnTargetInfo() const
+{
+	return FindBestTargetInfo(LaunchOnTarget, true);
+}
+
+#undef LOCTEXT_NAMESPACE
+
diff --git a/Engine/Source/Developer/DeveloperToolSettings/Private/ProjectPackagingSettings.cpp b/Engine/Source/Developer/DeveloperToolSettings/Private/ProjectPackagingSettings.cpp
--- a/Engine/Source/Developer/DeveloperToolSettings/Private/ProjectPackagingSettings.cpp
+++ b/Engine/Source/Developer/DeveloperToolSettings/Private/ProjectPackagingSettings.cpp
@@ -83,13 +83,6 @@ void UProjectPackagingSettings::PostEditChangeProperty(FPropertyChangedEvent& Pr
 		// We need to fix paths for no name updates to catch the reloadconfig call
 		FixCookingPaths();
 	}
-	else if (Name == FName((TEXT("StagingDirectory"))))
-	{
-		// fix up path
-		FString Path = StagingDirectory.Path;
-		FPaths::MakePathRelativeTo(Path, FPlatformProcess::BaseDir());
-		StagingDirectory.Path = Path;
-	}
 	else if (Name == FName(TEXT("ForDistribution")))
 	{
 		if (ForDistribution && BuildConfiguration != EProjectPackagingBuildConfigurations::PPBC_Shipping)
@@ -196,7 +189,7 @@ TArray<EProjectPackagingBuildConfigurations> UProjectPackagingSettings::GetValid
 }
 
 
-static const FTargetInfo* FindBestTargetInfo(const FString& TargetName, bool bContentOnlyUsesEngineTargets, bool* bOutIsProjectTarget=nullptr)
+const FTargetInfo* FindBestTargetInfo(const FString& TargetName, bool bContentOnlyUsesEngineTargets, bool* bOutIsProjectTarget=nullptr)
 {
 	bool bUseEngineTargets = false;
 	if (bContentOnlyUsesEngineTargets)
@@ -235,74 +228,11 @@ static const FTargetInfo* FindBestTargetInfo(const FString& TargetName, bool bCo
 
 }
 
+
 const FTargetInfo* UProjectPackagingSettings::GetBuildTargetInfo() const
 {
 	return FindBestTargetInfo(BuildTarget, false);
 }
 
-const FTargetInfo* UProjectPackagingSettings::GetBuildTargetInfoForPlatform(FName PlatformName, bool& bOutIsProjectTarget) const
-{
-	return FindBestTargetInfo(GetBuildTargetForPlatform(PlatformName), true, &bOutIsProjectTarget);
-}
-
-const FTargetInfo* UProjectPackagingSettings::GetLaunchOnTargetInfo() const
-{
-	return FindBestTargetInfo(LaunchOnTarget, true);
-}
-
-
-EProjectPackagingBuildConfigurations UProjectPackagingSettings::GetBuildConfigurationForPlatform(FName PlatformName) const
-{
-	const EProjectPackagingBuildConfigurations* Value = PerPlatformBuildConfig.Find(PlatformName);
-
-	// PPBC_MAX defines the default project setting case and should be handled accordingly.
-	return Value == nullptr ? EProjectPackagingBuildConfigurations::PPBC_MAX : *Value;
-}
-
-void UProjectPackagingSettings::SetBuildConfigurationForPlatform(FName PlatformName, EProjectPackagingBuildConfigurations Configuration)
-{
-	if (Configuration == EProjectPackagingBuildConfigurations::PPBC_MAX)
-	{
-		PerPlatformBuildConfig.Remove(PlatformName);
-	}
-	else
-	{
-		PerPlatformBuildConfig.Add(PlatformName, Configuration);
-	}
-}
-
-FName UProjectPackagingSettings::GetTargetFlavorForPlatform(FName FlavorName) const
-{
-	const FName* Value = PerPlatformTargetFlavorName.Find(FlavorName);
-
-	// the flavor name is also the name of the vanilla info
-	return Value == nullptr ? FlavorName : *Value;
-}
-
-void UProjectPackagingSettings::SetTargetFlavorForPlatform(FName PlatformName, FName TargetFlavorName)
-{
-	PerPlatformTargetFlavorName.Add(PlatformName, TargetFlavorName);
-}
-
-FString UProjectPackagingSettings::GetBuildTargetForPlatform(FName PlatformName) const
-{
-	const FString* Value = PerPlatformBuildTarget.Find(PlatformName);
-
-	// empty string defines the default project setting case and should be handled accordingly.
-	return Value == nullptr ? "" : *Value;
-}
-
-void UProjectPackagingSettings::SetBuildTargetForPlatform(FName PlatformName, FString BuildTargetName)
-{
-	if (BuildTargetName.IsEmpty())
-	{
-		PerPlatformBuildTarget.Remove(PlatformName);
-	}
-	else
-	{
-		PerPlatformBuildTarget.Add(PlatformName, BuildTargetName);
-	}
-}
-
 #undef LOCTEXT_NAMESPACE
 
diff --git a/Engine/Source/Editor/TurnkeySupport/Private/TurnkeyEditorSupport.cpp b/Engine/Source/Editor/TurnkeySupport/Private/TurnkeyEditorSupport.cpp
--- a/Engine/Source/Editor/TurnkeySupport/Private/TurnkeyEditorSupport.cpp
+++ b/Engine/Source/Editor/TurnkeySupport/Private/TurnkeyEditorSupport.cpp
@@ -26,6 +26,7 @@
 #include "ToolMenus.h"
 #include "FileHelpers.h"
 #include "Settings/ProjectPackagingSettings.h"
+#include "Settings/PlatformsMenuSettings.h"
 #include "Containers/Set.h"
 
 namespace
@@ -99,7 +100,7 @@ void FTurnkeyEditorSupport::LaunchRunningMap(const FString& DeviceId, const FStr
 		}
 		else
 		{
-			PlatformInfo = PlatformInfo::FindPlatformInfo(GetDefault<UProjectPackagingSettings>()->GetTargetFlavorForPlatform(*TargetDeviceId.GetPlatformName()));
+			PlatformInfo = PlatformInfo::FindPlatformInfo(GetDefault<UPlatformsMenuSettings>()->GetTargetFlavorForPlatform(*TargetDeviceId.GetPlatformName()));
 		}
 					
 		FString UBTPlatformName = PlatformInfo->DataDrivenPlatformInfo->UBTPlatformString;
diff --git a/Engine/Source/Editor/TurnkeySupport/Private/TurnkeySupportModule.cpp b/Engine/Source/Editor/TurnkeySupport/Private/TurnkeySupportModule.cpp
--- a/Engine/Source/Editor/TurnkeySupport/Private/TurnkeySupportModule.cpp
+++ b/Engine/Source/Editor/TurnkeySupport/Private/TurnkeySupportModule.cpp
@@ -31,7 +31,8 @@
 #include "ISourceControlModule.h"
 #include "ITargetDeviceProxy.h"
 #include "ITargetDeviceServicesModule.h"
-#include "Settings/ProjectPackagingSettings.h" // can go away once turnkey menu handles Package
+#include "Settings/ProjectPackagingSettings.h"
+#include "Settings/PlatformsMenuSettings.h"
 #include "IDesktopPlatform.h"
 #include "DesktopPlatformModule.h"
 #include "PlatformInfo.h"
@@ -130,7 +131,7 @@ protected:
 	}
 
 
-	static bool ShouldBuildProject(UProjectPackagingSettings* PackagingSettings, const ITargetPlatform* TargetPlatform)
+	static bool ShouldBuildProject(const UProjectPackagingSettings* PackagingSettings, const ITargetPlatform* TargetPlatform)
 	{
 		const UProjectPackagingSettings::FConfigurationInfo& ConfigurationInfo = UProjectPackagingSettings::ConfigurationInfo[(int)PackagingSettings->BuildConfiguration];
 
@@ -268,8 +269,9 @@ public:
 		TArray<FAnalyticsEventAttribute> AnalyticsParamArray;
 
 		// get a in-memory defaults which will have the user-settings, like the per-platform config/target platform stuff
-		UProjectPackagingSettings* AllPlatformPackagingSettings = GetMutableDefault<UProjectPackagingSettings>();
-	
+		const UProjectPackagingSettings* PackagingSettings = GetPackagingSettingsForPlatform(IniPlatformName);
+		UPlatformsMenuSettings* PlatformsSettings = GetMutableDefault<UPlatformsMenuSettings>();
+
 		// installed builds only support standard Game type builds (not Client, Server, etc) so instead of looking up a setting that the user can't set, 
 		// always use the base PlatformInfo for Game builds, which will be named the same as the platform itself
 		const PlatformInfo::FTargetPlatformInfo* PlatformInfo = nullptr;
@@ -279,7 +281,7 @@ public:
 		}
 		else
 		{
-			PlatformInfo = PlatformInfo::FindPlatformInfo(AllPlatformPackagingSettings->GetTargetFlavorForPlatform(IniPlatformName));
+			PlatformInfo = PlatformInfo::FindPlatformInfo(PlatformsSettings->GetTargetFlavorForPlatform(IniPlatformName));
 		}
 		// this is unexpected to be able to happen, but it could if there was a bad value saved in the UProjectPackagingSettings - if this trips, we should handle errors
 		check(PlatformInfo != nullptr);
@@ -313,9 +315,6 @@ public:
 		// this may delete UProjectPackagingSettings , don't hold it across this call
 		FTurnkeyEditorSupport::SaveAll();
 
-		// get all the helper objects
-		UProjectPackagingSettings* PackagingSettings = GetPackagingSettingsForPlatform(IniPlatformName);
-
 		// basic BuildCookRun params we always want
 		FString BuildCookRunParams = FString::Printf(TEXT("-nop4 -utf8output %s -cook "), GetUATCompilationFlags());
 
@@ -327,7 +326,7 @@ public:
 		}
 
 		bool bIsProjectBuildTarget = false;
-		const FTargetInfo* BuildTargetInfo = AllPlatformPackagingSettings->GetBuildTargetInfoForPlatform(IniPlatformName, bIsProjectBuildTarget);
+		const FTargetInfo* BuildTargetInfo = PlatformsSettings->GetBuildTargetInfoForPlatform(IniPlatformName, bIsProjectBuildTarget);
 
 		// Only add the -Target=... argument for code projects. Content projects will return UnrealGame/UnrealClient/UnrealServer here, but
 		// may need a temporary target generated to enable/disable plugins. Specifying -Target in these cases will cause packaging to fail,
@@ -394,20 +393,20 @@ public:
 			ContentPrepIcon = FAppStyle::Get().GetBrush(TEXT("MainFrame.PackageProject"));
 
 			// let the user pick a target directory
-			if (AllPlatformPackagingSettings->StagingDirectory.Path.IsEmpty())
+			if (PlatformsSettings->StagingDirectory.Path.IsEmpty())
 			{
-				AllPlatformPackagingSettings->StagingDirectory.Path = FPaths::ProjectDir();
+				PlatformsSettings->StagingDirectory.Path = FPaths::ProjectDir();
 			}
 
 			FString OutFolderName;
 
-			if (!FDesktopPlatformModule::Get()->OpenDirectoryDialog(FSlateApplication::Get().FindBestParentWindowHandleForDialogs(nullptr), LOCTEXT("PackageDirectoryDialogTitle", "Package project...").ToString(), AllPlatformPackagingSettings->StagingDirectory.Path, OutFolderName))
+			if (!FDesktopPlatformModule::Get()->OpenDirectoryDialog(FSlateApplication::Get().FindBestParentWindowHandleForDialogs(nullptr), LOCTEXT("PackageDirectoryDialogTitle", "Package project...").ToString(), PlatformsSettings->StagingDirectory.Path, OutFolderName))
 			{
 				return;
 			}
 
-			AllPlatformPackagingSettings->StagingDirectory.Path = OutFolderName;
-			AllPlatformPackagingSettings->SaveConfig();
+			PlatformsSettings->StagingDirectory.Path = OutFolderName;
+			PlatformsSettings->SaveConfig();
 
 
 			BuildCookRunParams += TEXT(" -stage -archive -package");
@@ -423,9 +422,7 @@ public:
 			}
 
 			// Pak file(s) must be used when using container file(s)
-			PackagingSettings->UsePakFile |= PackagingSettings->bUseIoStore;
-
-			if (PackagingSettings->UsePakFile)
+			if (PackagingSettings->UsePakFile || PackagingSettings->bUseIoStore)
 			{
 				BuildCookRunParams += TEXT(" -pak");
 				if (PackagingSettings->bUseIoStore)
@@ -453,7 +450,7 @@ public:
 				BuildCookRunParams += TEXT(" -applocaldirectory=\"$(EngineDir)/Binaries/ThirdParty/AppLocalDependencies\"");
 			}
 
-			BuildCookRunParams += FString::Printf(TEXT(" -archivedirectory=\"%s\""), *AllPlatformPackagingSettings->StagingDirectory.Path);
+			BuildCookRunParams += FString::Printf(TEXT(" -archivedirectory=\"%s\""), *PlatformsSettings->StagingDirectory.Path);
 
 			if (PackagingSettings->ForDistribution)
 			{
@@ -476,11 +473,11 @@ public:
 				BuildCookRunParams += FString::Printf(TEXT(" -manifests -createchunkinstall -chunkinstalldirectory=\"%s\" -chunkinstallversion=%s"), *(PackagingSettings->HttpChunkInstallDataDirectory.Path), *(PackagingSettings->HttpChunkInstallDataVersion));
 			}
 
-			EProjectPackagingBuildConfigurations BuildConfig = AllPlatformPackagingSettings->GetBuildConfigurationForPlatform(IniPlatformName);
+			EProjectPackagingBuildConfigurations BuildConfig = PlatformsSettings->GetBuildConfigurationForPlatform(IniPlatformName);
 			// if PPBC_MAX is set, then the project default should be used instead of the per platform build config
 			if (BuildConfig == EProjectPackagingBuildConfigurations::PPBC_MAX)
 			{
-				BuildConfig = AllPlatformPackagingSettings->BuildConfiguration;
+				BuildConfig = PackagingSettings->BuildConfiguration;
 			}
 
 			// when distribution is set, always package in shipping, which overrides the per platform build config
@@ -600,29 +597,29 @@ public:
 
 		// pass the editor setting down to UAT in case they aren't saved to the ini's that UAT would read
 		{
-			UProjectPackagingSettings* PackagingSettings = FTurnkeySupportCallbacks::GetPackagingSettingsForPlatform(IniPlatformName);
-			UProjectPackagingSettings* AllPlatformPackagingSettings = GetMutableDefault<UProjectPackagingSettings>();
+			UProjectPackagingSettings* PlatformPackagingSettings = FTurnkeySupportCallbacks::GetPackagingSettingsForPlatform(IniPlatformName);
+			const UPlatformsMenuSettings* PlatformsSettings = GetDefault<UPlatformsMenuSettings>();
 			bool bIsProjectBuildTarget = false;
-			const FTargetInfo* BuildTargetInfo = AllPlatformPackagingSettings->GetBuildTargetInfoForPlatform(IniPlatformName, bIsProjectBuildTarget);
+			const FTargetInfo* BuildTargetInfo = PlatformsSettings->GetBuildTargetInfoForPlatform(IniPlatformName, bIsProjectBuildTarget);
 
 			CommandLine.Appendf(TEXT(" -overridetarget=%s"), *BuildTargetInfo->Name);
 
 			// distribution builds can set shipping 
-			EProjectPackagingBuildConfigurations BuildConfig = PackagingSettings->ForDistribution ? 
+			EProjectPackagingBuildConfigurations BuildConfig = PlatformPackagingSettings->ForDistribution ?
 				EProjectPackagingBuildConfigurations::PPBC_Shipping :
-				AllPlatformPackagingSettings->GetBuildConfigurationForPlatform(IniPlatformName);
+			PlatformsSettings->GetBuildConfigurationForPlatform(IniPlatformName);
 
 			// if PPBC_MAX is set, then the project default should be used instead of the per platform build config
 			if (BuildConfig == EProjectPackagingBuildConfigurations::PPBC_MAX)
 			{
-				BuildConfig = AllPlatformPackagingSettings->BuildConfiguration;
+				BuildConfig = PlatformPackagingSettings->BuildConfiguration;
 			}
 			const UProjectPackagingSettings::FConfigurationInfo& ConfigurationInfo = UProjectPackagingSettings::ConfigurationInfo[(int)BuildConfig];
 			CommandLine.Appendf(TEXT(" -overrideconfiguration=%s"), LexToString(ConfigurationInfo.Configuration));
 
 
 			// get the chosen cook flavor (texture format, etc)
-			const PlatformInfo::FTargetPlatformInfo* PlatformInfo = PlatformInfo::FindPlatformInfo(AllPlatformPackagingSettings->GetTargetFlavorForPlatform(IniPlatformName));
+			const PlatformInfo::FTargetPlatformInfo* PlatformInfo = PlatformInfo::FindPlatformInfo(PlatformsSettings->GetTargetFlavorForPlatform(IniPlatformName));
 			if (PlatformInfo->IsFlavor())
 			{
 				CommandLine.Appendf(TEXT(" -overrideflavor=%s"), *PlatformInfo->PlatformFlavor.ToString());
@@ -634,7 +631,7 @@ public:
 
 	static void ExecuteCustomBuild(FName IniPlatformName, FString DeviceId, FProjectBuildSettings Build)
 	{
-		const PlatformInfo::FTargetPlatformInfo* PlatformInfo = PlatformInfo::FindPlatformInfo(GetDefault<UProjectPackagingSettings>()->GetTargetFlavorForPlatform(IniPlatformName));
+		const PlatformInfo::FTargetPlatformInfo* PlatformInfo = PlatformInfo::FindPlatformInfo(GetDefault<UPlatformsMenuSettings>()->GetTargetFlavorForPlatform(IniPlatformName));
 		const FString ProjectPath = GetProjectPathForTurnkey();
 
 		// throw this on the command before the actual automation name (Turnkey, BuildCOokRun, etc) because they are global-to-UAT options, not options for a single command
@@ -683,15 +680,15 @@ public:
 			if (Build.BuildCookRunParams.Contains("{BrowseForDir}"))
 			{
 				// we reuse the already-presenbt StagingDirectory to save it, although it's more like "BuildOutpuDirectory" now 
-				UProjectPackagingSettings* AllPlatformPackagingSettings = GetMutableDefault<UProjectPackagingSettings>();
+				UPlatformsMenuSettings* PlatformsSettings = GetMutableDefault<UPlatformsMenuSettings>();
 				FString OutFolderName;
-				if (!FDesktopPlatformModule::Get()->OpenDirectoryDialog(FSlateApplication::Get().FindBestParentWindowHandleForDialogs(nullptr), LOCTEXT("PackageDirectoryDialogTitle", "Package project...").ToString(), AllPlatformPackagingSettings->StagingDirectory.Path, OutFolderName))
+				if (!FDesktopPlatformModule::Get()->OpenDirectoryDialog(FSlateApplication::Get().FindBestParentWindowHandleForDialogs(nullptr), LOCTEXT("PackageDirectoryDialogTitle", "Package project...").ToString(), PlatformsSettings->StagingDirectory.Path, OutFolderName))
 				{
 					return;
 				}
 
-				AllPlatformPackagingSettings->StagingDirectory.Path = OutFolderName;
-				AllPlatformPackagingSettings->SaveConfig();
+				PlatformsSettings->StagingDirectory.Path = OutFolderName;
+				PlatformsSettings->SaveConfig();
 
 				CommandLine += FString::Printf(TEXT(" -outputdir=\"%s\""), *OutFolderName);
 			}
@@ -703,21 +700,21 @@ public:
 
 	static void SetPackageBuildConfiguration(const PlatformInfo::FTargetPlatformInfo* Info, EProjectPackagingBuildConfigurations BuildConfiguration)
 	{
-		UProjectPackagingSettings* PackagingSettings = GetMutableDefault<UProjectPackagingSettings>();
-		PackagingSettings->SetBuildConfigurationForPlatform(Info->IniPlatformName, BuildConfiguration);
-		PackagingSettings->SaveConfig();
+		UPlatformsMenuSettings* PlatformsSettings = GetMutableDefault<UPlatformsMenuSettings>();
+		PlatformsSettings->SetBuildConfigurationForPlatform(Info->IniPlatformName, BuildConfiguration);
+		PlatformsSettings->SaveConfig();
 	}
 
 	static bool PackageBuildConfigurationIsChecked(const PlatformInfo::FTargetPlatformInfo* Info, EProjectPackagingBuildConfigurations BuildConfiguration)
 	{
-		return GetDefault<UProjectPackagingSettings>()->GetBuildConfigurationForPlatform(Info->IniPlatformName) == BuildConfiguration;
+		return GetDefault<UPlatformsMenuSettings>()->GetBuildConfigurationForPlatform(Info->IniPlatformName) == BuildConfiguration;
 	}	
 	
 	static void SetActiveFlavor(const PlatformInfo::FTargetPlatformInfo* Info)
 	{
-		UProjectPackagingSettings* PackagingSettings = GetMutableDefault<UProjectPackagingSettings>();
-		PackagingSettings->SetTargetFlavorForPlatform(Info->IniPlatformName, Info->Name);
-		PackagingSettings->SaveConfig();
+		UPlatformsMenuSettings* PlatformsSettings = GetMutableDefault<UPlatformsMenuSettings>();
+		PlatformsSettings->SetTargetFlavorForPlatform(Info->IniPlatformName, Info->Name);
+		PlatformsSettings->SaveConfig();
 	}
 
 	static bool CanSetActiveFlavor(const PlatformInfo::FTargetPlatformInfo* Info)
@@ -727,31 +724,31 @@ public:
 
 	static bool SetActiveFlavorIsChecked(const PlatformInfo::FTargetPlatformInfo* Info)
 	{
-		return GetDefault<UProjectPackagingSettings>()->GetTargetFlavorForPlatform(Info->IniPlatformName) == Info->Name;
+		return GetDefault<UPlatformsMenuSettings>()->GetTargetFlavorForPlatform(Info->IniPlatformName) == Info->Name;
 	}
 
 	static void SetPackageBuildTarget(const PlatformInfo::FTargetPlatformInfo* Info, FString TargetName)
 	{
-		UProjectPackagingSettings* PackagingSettings = GetMutableDefault<UProjectPackagingSettings>();
-		PackagingSettings->SetBuildTargetForPlatform(Info->IniPlatformName, TargetName);
-		PackagingSettings->SaveConfig();
+		UPlatformsMenuSettings* PlatformsSettings = GetMutableDefault<UPlatformsMenuSettings>();
+		PlatformsSettings->SetBuildTargetForPlatform(Info->IniPlatformName, TargetName);
+		PlatformsSettings->SaveConfig();
 	}
 
 	static bool PackageBuildTargetIsChecked(const PlatformInfo::FTargetPlatformInfo* Info, FString TargetName)
 	{
-		return GetDefault<UProjectPackagingSettings>()->GetBuildTargetForPlatform(Info->IniPlatformName) == TargetName;
+		return GetDefault<UPlatformsMenuSettings>()->GetBuildTargetForPlatform(Info->IniPlatformName) == TargetName;
 	}
 
 	static void SetLaunchOnBuildTarget(FString TargetName)
 	{
-		UProjectPackagingSettings* PackagingSettings = GetMutableDefault<UProjectPackagingSettings>();
-		PackagingSettings->LaunchOnTarget = TargetName;
-		PackagingSettings->SaveConfig();
+		UPlatformsMenuSettings* PlatformsSettings = GetMutableDefault<UPlatformsMenuSettings>();
+		PlatformsSettings->LaunchOnTarget = TargetName;
+		PlatformsSettings->SaveConfig();
 	}
 
 	static bool LaunchOnBuildTargetIsChecked(FString TargetName)
 	{
-		return GetDefault<UProjectPackagingSettings>()->GetLaunchOnTargetInfo()->Name == TargetName;
+		return GetDefault<UPlatformsMenuSettings>()->GetLaunchOnTargetInfo()->Name == TargetName;
 	}
 
 	static void SetCookOnTheFly()
@@ -1625,7 +1622,7 @@ static void GenerateDeviceProxyMenuParams(TSharedPtr<ITargetDeviceProxy> DeviceP
 				}
 				else
 				{
-					PlatformInfo = PlatformInfo::FindPlatformInfo(GetDefault<UProjectPackagingSettings>()->GetTargetFlavorForPlatform(PlatformName));
+					PlatformInfo = PlatformInfo::FindPlatformInfo(GetDefault<UPlatformsMenuSettings>()->GetTargetFlavorForPlatform(PlatformName));
 				}
 
 				FString VariantName = PlatformInfo->Name.ToString();
@@ -1635,7 +1632,7 @@ static void GenerateDeviceProxyMenuParams(TSharedPtr<ITargetDeviceProxy> DeviceP
 				if (DesktopPlatform->GetTargetsForCurrentProject().Num() > 0)
 				{
 					TArray<FTargetInfo> Targets = DesktopPlatform->GetTargetsForCurrentProject();
-					const FTargetInfo* TargetInfo = GetDefault<UProjectPackagingSettings>()->GetLaunchOnTargetInfo();
+					const FTargetInfo* TargetInfo = GetDefault<UPlatformsMenuSettings>()->GetLaunchOnTargetInfo();
 
 					const TArray<FTargetInfo> ClientTargets = Targets.FilterByPredicate([](const FTargetInfo& Target)
 						{
@@ -2298,8 +2295,7 @@ void FTurnkeySupportModule::UpdateSdkInfo()
 							// if platform uses flavors and doesn't already have a default, set default flavor to the first availabl
 							if (ValidFlavors.Num() > 1)
 							{
-								UProjectPackagingSettings* AllPlatformPackagingSettings = GetMutableDefault<UProjectPackagingSettings>();
-								FName CurrentFlavor = AllPlatformPackagingSettings->GetTargetFlavorForPlatform(PlatformName);
+								FName CurrentFlavor = GetDefault<UPlatformsMenuSettings>()->GetTargetFlavorForPlatform(PlatformName);
 
 								if (CurrentFlavor == VanillaInfo->Name)
 								{
diff --git a/Engine/Source/Editor/UnrealEd/Private/PlayLevelViaLauncher.cpp b/Engine/Source/Editor/UnrealEd/Private/PlayLevelViaLauncher.cpp
--- a/Engine/Source/Editor/UnrealEd/Private/PlayLevelViaLauncher.cpp
+++ b/Engine/Source/Editor/UnrealEd/Private/PlayLevelViaLauncher.cpp
@@ -23,6 +23,7 @@
 #include "Settings/EditorExperimentalSettings.h"
 #include "Settings/LevelEditorPlaySettings.h"
 #include "Settings/ProjectPackagingSettings.h"
+#include "Settings/PlatformsMenuSettings.h"
 #include "TargetReceipt.h"
 #include "UnrealEdMisc.h"
 #include "Widgets/Notifications/SNotificationList.h"
@@ -157,8 +158,9 @@ void UEditorEngine::StartPlayUsingLauncherSession(FRequestPlaySessionParams& InR
 	default:
 	{
 		const UProjectPackagingSettings* AllPlatformPackagingSettings = GetDefault<UProjectPackagingSettings>();
+		const UPlatformsMenuSettings* PlatformsSettings = GetDefault<UPlatformsMenuSettings>();
 
-		EProjectPackagingBuildConfigurations BuildConfig = AllPlatformPackagingSettings->GetBuildConfigurationForPlatform(*IniPlatformName);
+		EProjectPackagingBuildConfigurations BuildConfig = PlatformsSettings->GetBuildConfigurationForPlatform(*IniPlatformName);
 		// if PPBC_MAX is set, then the project default should be used instead of the per platform build config
 		if (BuildConfig == EProjectPackagingBuildConfigurations::PPBC_MAX)
 		{
@@ -280,7 +282,7 @@ void UEditorEngine::StartPlayUsingLauncherSession(FRequestPlaySessionParams& InR
 	// only set the BuildTarget in code-based projects
 	if (LauncherSessionInfo->bPlayUsingLauncherHasCode)
 	{
-		const FTargetInfo* TargetInfo = GetDefault<UProjectPackagingSettings>()->GetLaunchOnTargetInfo();
+		const FTargetInfo* TargetInfo = GetDefault<UPlatformsMenuSettings>()->GetLaunchOnTargetInfo();
 		if (TargetInfo != nullptr)
 		{
 			LauncherProfile->SetBuildTarget(TargetInfo->Name);
