- Changed target selection in Quick Launch when there is no game target to look for a client target

diff --git a/Engine/Source/Editor/TurnkeySupport/Private/TurnkeySupportModule.cpp b/Engine/Source/Editor/TurnkeySupport/Private/TurnkeySupportModule.cpp
--- a/Engine/Source/Editor/TurnkeySupport/Private/TurnkeySupportModule.cpp
+++ b/Engine/Source/Editor/TurnkeySupport/Private/TurnkeySupportModule.cpp
@@ -1279,16 +1279,26 @@ static void MakeTurnkeyPlatformMenu(UToolMenu* ToolMenu, FName IniPlatformName,
 
 			if (ValidTargets.Num() > 1)
 			{
-				// Set BuildTarget to default to Game if it hasn't been set
-				if(AllPlatformPackagingSettings->BuildTarget.IsEmpty())
+				// Set BuildTarget to default to Game if it hasn't been set (or if it was set, but is no longer valid)
+				FString BuildTarget = AllPlatformPackagingSettings->BuildTarget;
+				if (BuildTarget.IsEmpty() ||
+				   !ValidTargets.ContainsByPredicate([&BuildTarget](const FTargetInfo& Target) { return Target.Name == BuildTarget; } ))
 				{
-					const TArray<FTargetInfo> GameTarget = ValidTargets.FilterByPredicate([](const FTargetInfo& Target)
-						{
-							return Target.Type == EBuildTargetType::Game;
-						});
-					check(GameTarget.Num() > 0);
+					TArray<FTargetInfo> GameTargets = ValidTargets.FilterByPredicate([](const FTargetInfo& Target)
+						{ return Target.Type == EBuildTargetType::Game;	});
+					// if there are no Game targets, look for a Client target
+					if (GameTargets.Num() == 0)
+					{
+						GameTargets = ValidTargets.FilterByPredicate([](const FTargetInfo& Target)
+							{ return Target.Type == EBuildTargetType::Client;	});
+					}
+					// if no clients, then just _anything_ 
+					if (GameTargets.Num() == 0)
+					{
+						GameTargets = ValidTargets;
+					}
 
-					AllPlatformPackagingSettings->BuildTarget = GameTarget[0].Name;
+					AllPlatformPackagingSettings->BuildTarget = GameTargets[0].Name;
 					AllPlatformPackagingSettings->SaveConfig();
 				}
 
