Only iOS and tvOS currently support device connection type specific icons

#rnx

diff --git a/Engine/Source/Developer/IOS/IOSTargetPlatform/Private/IOSTargetPlatform.cpp b/Engine/Source/Developer/IOS/IOSTargetPlatform/Private/IOSTargetPlatform.cpp
--- a/Engine/Source/Developer/IOS/IOSTargetPlatform/Private/IOSTargetPlatform.cpp
+++ b/Engine/Source/Developer/IOS/IOSTargetPlatform/Private/IOSTargetPlatform.cpp
@@ -504,6 +504,9 @@ bool FIOSTargetPlatform::SupportsFeature( ETargetPlatformFeatures Feature ) cons
 		case ETargetPlatformFeatures::ShowAsPlatformGroup:
 			return false;
 
+		case ETargetPlatformFeatures::SupportsMultipleConnectionTypes:
+			return true;
+
 		default:
 			break;
 	}
diff --git a/Engine/Source/Developer/TargetPlatform/Public/Interfaces/ITargetPlatform.h b/Engine/Source/Developer/TargetPlatform/Public/Interfaces/ITargetPlatform.h
--- a/Engine/Source/Developer/TargetPlatform/Public/Interfaces/ITargetPlatform.h
+++ b/Engine/Source/Developer/TargetPlatform/Public/Interfaces/ITargetPlatform.h
@@ -115,7 +115,10 @@ enum class ETargetPlatformFeatures
 	NormalmapLAEncodingMode,
 
 	/** All devices of this platform should be grouped under one platform group */
-	ShowAsPlatformGroup
+	ShowAsPlatformGroup,
+
+	/** Does the platform allow various connection types to be used (ie: wifi and usb) */
+	SupportsMultipleConnectionTypes
 };
 
 enum class EPlatformAuthentication
diff --git a/Engine/Source/Editor/EditorStyle/Private/StarshipStyle.cpp b/Engine/Source/Editor/EditorStyle/Private/StarshipStyle.cpp
--- a/Engine/Source/Editor/EditorStyle/Private/StarshipStyle.cpp
+++ b/Engine/Source/Editor/EditorStyle/Private/StarshipStyle.cpp
@@ -7299,6 +7299,8 @@ void FStarshipEditorStyle::FStyle::SetupAutomationStyles()
 
 		Set( "DeviceDetails.WIFI.IOS", new IMAGE_BRUSH("Starship/DeviceManager/AppleWifi_128x", Icon128x128));
 		Set( "DeviceDetails.USB.IOS", new IMAGE_BRUSH("Starship/DeviceManager/AppleUsb_128x", Icon128x128));
+		Set( "DeviceDetails.WIFI.TVOS", new IMAGE_BRUSH("Starship/DeviceManager/TVOS_wifi_128x", Icon128x128));
+		Set( "DeviceDetails.USB.TVOS", new IMAGE_BRUSH("Starship/DeviceManager/TVOS_usb_128x", Icon128x128));
 	}
 
 	// Settings Editor
diff --git a/Engine/Source/Editor/TurnkeySupport/Private/TurnkeySupportModule.cpp b/Engine/Source/Editor/TurnkeySupport/Private/TurnkeySupportModule.cpp
--- a/Engine/Source/Editor/TurnkeySupport/Private/TurnkeySupportModule.cpp
+++ b/Engine/Source/Editor/TurnkeySupport/Private/TurnkeySupportModule.cpp
@@ -977,7 +977,9 @@ static FSlateIcon MakePlatformSdkIconAttribute(FName IniPlatformName, TSharedPtr
 	}
 	else
 	{
-		if (DeviceProxy != nullptr && !DeviceProxy->GetConnectionType().IsEmpty())
+		const ITargetPlatform* Platform = GetTargetPlatformManager()->FindTargetPlatform(IniPlatformName);
+		if ((DeviceProxy != nullptr) && (Platform != nullptr) && !DeviceProxy->GetConnectionType().IsEmpty() &&
+			Platform->SupportsFeature(ETargetPlatformFeatures::SupportsMultipleConnectionTypes))
 		{
 			FSlateIcon Icon;
 			if (DeviceProxy->GetConnectionType().Contains(TEXT("Network")))
