fix for crash if the editor is shut down immediately after starting

diff --git a/Engine/Source/Editor/TurnkeySupport/Private/TurnkeySupportModule.cpp b/Engine/Source/Editor/TurnkeySupport/Private/TurnkeySupportModule.cpp
--- a/Engine/Source/Editor/TurnkeySupport/Private/TurnkeySupportModule.cpp
+++ b/Engine/Source/Editor/TurnkeySupport/Private/TurnkeySupportModule.cpp
@@ -2575,7 +2575,7 @@ void FTurnkeySupportModule::UpdateSdkInfoForDevices(TArray<FString> PlatformDevi
 				    }
 			    }
 			}
-			else
+			else if (!IsEngineExitRequested())
 			{
 			    for (const FString& Id : PlatformDeviceIds)
 			    {
