Display available iOS Simulators under the Platforms menu dropdown
Part 2 of 3

diff --git a/Engine/Source/Developer/TargetDeviceServices/Private/Proxies/TargetDeviceProxy.cpp b/Engine/Source/Developer/TargetDeviceServices/Private/Proxies/TargetDeviceProxy.cpp
--- a/Engine/Source/Developer/TargetDeviceServices/Private/Proxies/TargetDeviceProxy.cpp
+++ b/Engine/Source/Developer/TargetDeviceServices/Private/Proxies/TargetDeviceProxy.cpp
@@ -60,6 +60,7 @@ void FTargetDeviceProxy::UpdateFromMessage( const FTargetDeviceServicePong& Mess
 	HostUser = Message.HostUser;
 	Make = Message.Make;
 	Model = Message.Model;
+	Type = Message.Type;
 	Name = Aggregated? Message.AllDevicesName: Message.Name;
 	DeviceUser = Message.DeviceUser;
 	DeviceUserPassword = Message.DeviceUserPassword;
diff --git a/Engine/Source/Developer/TargetDeviceServices/Private/Proxies/TargetDeviceProxy.h b/Engine/Source/Developer/TargetDeviceServices/Private/Proxies/TargetDeviceProxy.h
--- a/Engine/Source/Developer/TargetDeviceServices/Private/Proxies/TargetDeviceProxy.h
+++ b/Engine/Source/Developer/TargetDeviceServices/Private/Proxies/TargetDeviceProxy.h
@@ -185,6 +185,11 @@ public:
 		return Shared;
 	}
 	
+	virtual bool IsSimulated() const override
+	{
+		return (ConnectionType == TEXT("Simulator"));
+	}
+	
 	virtual bool TerminateLaunchedProcess(FName InVariant, const FString& ProcessIdentifier) override;
 
 	virtual void PowerOff(bool Force) override;
diff --git a/Engine/Source/Developer/TargetDeviceServices/Private/Services/TargetDeviceService.cpp b/Engine/Source/Developer/TargetDeviceServices/Private/Services/TargetDeviceService.cpp
--- a/Engine/Source/Developer/TargetDeviceServices/Private/Services/TargetDeviceService.cpp
+++ b/Engine/Source/Developer/TargetDeviceServices/Private/Services/TargetDeviceService.cpp
@@ -426,7 +426,7 @@ void FTargetDeviceService::HandlePingMessage(const FTargetDeviceServicePing& InM
 		Message->HostName = FPlatformProcess::ComputerName();
 		Message->HostUser = FPlatformProcess::UserName(false);
 		Message->Connected = DefaultDevice->IsConnected();
-		Message->ConnectionType = (DefaultDevice->GetDeviceConnectionType() == ETargetDeviceConnectionTypes::Wifi) ? "Network" : "USB";
+		Message->ConnectionType = TargetDeviceConnectionTypes::ToString(DefaultDevice->GetDeviceConnectionType());
 		Message->Authorized = DefaultDevice->IsAuthorized();
 		Message->Make = TEXT("@todo");
 		Message->Model = DefaultDevice->GetModelId();
diff --git a/Engine/Source/Developer/TargetDeviceServices/Public/ITargetDeviceProxy.h b/Engine/Source/Developer/TargetDeviceServices/Public/ITargetDeviceProxy.h
--- a/Engine/Source/Developer/TargetDeviceServices/Public/ITargetDeviceProxy.h
+++ b/Engine/Source/Developer/TargetDeviceServices/Public/ITargetDeviceProxy.h
@@ -241,6 +241,13 @@ public:
 	 */
 	virtual bool IsShared() const = 0;
 
+	/**
+	 * Checks whether this device is being simulated on a Simulator
+	 *
+	 * @return true if the device is being simulated, false otherwise.
+	 */
+	virtual bool IsSimulated() const = 0;
+	
 	/**
 	* Cancel the application running on the device
 	 * @param InVariant Variant of the device
diff --git a/Engine/Source/Editor/TurnkeySupport/Private/TurnkeySupportModule.cpp b/Engine/Source/Editor/TurnkeySupport/Private/TurnkeySupportModule.cpp
--- a/Engine/Source/Editor/TurnkeySupport/Private/TurnkeySupportModule.cpp
+++ b/Engine/Source/Editor/TurnkeySupport/Private/TurnkeySupportModule.cpp
@@ -940,6 +940,7 @@ static void TurnkeyInstallSdk(FString IniPlatformName, bool bPreferFull, bool bF
 }
 
 static TMap<FName, FString> PerPlatformLastChosen;
+static TMap<FName, FString> PerPlatformLastSimChosen;
 
 static TAttribute<FText> MakeSdkStatusAttribute(FName IniPlatformName, TSharedPtr< ITargetDeviceProxy> DeviceProxy)
 {
@@ -992,6 +993,10 @@ static FSlateIcon MakePlatformSdkIconAttribute(FName IniPlatformName, TSharedPtr
 			{
 				Icon = FSlateIcon(FAppStyle::Get().GetStyleSetName(), *(TEXT("DeviceDetails.WIFI.") + IniPlatformName.ToString()));
 			}
+			else if (DeviceProxy->GetConnectionType().Contains(TEXT("Wifi")))
+			{
+				Icon = FSlateIcon(FAppStyle::Get().GetStyleSetName(), *(TEXT("DeviceDetails.WIFI.") + IniPlatformName.ToString()));
+			}
 			else if (DeviceProxy->GetConnectionType().Contains(TEXT("USB")))
 			{
 				Icon = FSlateIcon(FAppStyle::Get().GetStyleSetName(), *(TEXT("DeviceDetails.USB.") + IniPlatformName.ToString()));
@@ -1612,6 +1617,11 @@ static void LaunchOnDevice(const FString& DeviceId, const FString& DeviceName, b
 	FTurnkeyEditorSupport::LaunchRunningMap(DeviceId, DeviceName, GetProjectPathForTurnkey(), bUseTurnkey);
 }
 
+static void LaunchOnSimulator(const FString& DeviceId, const FString& DeviceName, bool bUseTurnkey)
+{								
+	// FTurnkeyEditorSupport::LaunchRunningMap(DeviceId, DeviceName, GetProjectPathForTurnkey(), false, true);
+}
+
 static void PrepareLaunchOn(FString DeviceId, FString DeviceName)
 {
 	FTurnkeyEditorSupport::PrepareToLaunchRunningMap(DeviceId, DeviceName);
@@ -1623,6 +1633,10 @@ static void HandleLaunchOnDeviceActionExecute(FString DeviceId, FString DeviceNa
 	LaunchOnDevice(DeviceId, DeviceName, bUseTurnkey);
 }
 
+static void HandleLaunchOnSimulatorActionExecute(FString DeviceId, FString DeviceName, bool bUseTurnkey)
+{
+	LaunchOnSimulator(DeviceId, DeviceName, bUseTurnkey);
+}
 
 static bool HandleLaunchOnDeviceActionCanExecute(FString DeviceName)
 {
@@ -1696,9 +1710,17 @@ static void GenerateDeviceProxyMenuParams(TSharedPtr<ITargetDeviceProxy> DeviceP
 					DeviceVariantName = *VariantName;
 				}
 
-				PerPlatformLastChosen[PlatformName] = DeviceProxy->GetName();
 				FString DeviceId = DeviceProxy->GetTargetDeviceId(DeviceVariantName);
-				HandleLaunchOnDeviceActionExecute(DeviceId, DeviceProxy->GetName(), true);
+				if (DeviceProxy->IsSimulated())
+				{
+					PerPlatformLastSimChosen[*DeviceProxy->GetType()] = DeviceProxy->GetName();
+					HandleLaunchOnSimulatorActionExecute(DeviceId, DeviceProxy->GetName(), true);
+				}
+				else
+				{
+					PerPlatformLastChosen[PlatformName] = DeviceProxy->GetName();
+					HandleLaunchOnDeviceActionExecute(DeviceId, DeviceProxy->GetName(), true);
+				}
 				ExternalOnClickDelegate.ExecuteIfBound(DeviceId);
 			}
 	//		, FCanExecuteAction::CreateStatic(&HandleLaunchOnDeviceActionCanExecute, DeviceProxy->GetName())
@@ -1719,7 +1741,7 @@ static void GenerateDeviceProxyMenuParams(TSharedPtr<ITargetDeviceProxy> DeviceP
 		TooltipArguments.Add(TEXT("ModelId"), FText::FromString(DeviceProxy->GetModel()));
 		TooltipArguments.Add(TEXT("OSVersion"), FText::FromString(DeviceProxy->GetOSVersion()));
 
-		OutTooltip = FText::Format(LOCTEXT("LaunchDeviceToolTipText_ThisDeviceExtra", "Launch on \"{DeviceID}\" [OSVer:'{OSVersion}' Model:'{ModelId}']"), TooltipArguments);
+		OutTooltip = FText::Format(LOCTEXT("LaunchDeviceToolTipText_ThisDeviceExtra", "Launch on \"{DeviceID}\" [OS:'{OSVersion}' Info:'{ModelId}']"), TooltipArguments);
 	}
 	else
 	{
@@ -1758,9 +1780,16 @@ void FTurnkeySupportModule::MakeQuickLaunchItems(class UToolMenu* Menu, FOnQuick
 			}
 
 			// look for devices for all platforms, even if the platform isn't installed - Turnkey can install Sdk after selecting LaunchOn
-			TArray<TSharedPtr<ITargetDeviceProxy>> DeviceProxies;
-			TargetDeviceServicesModule->GetDeviceProxyManager()->GetAllProxies(PlatformName, DeviceProxies);
+			TArray<TSharedPtr<ITargetDeviceProxy>> AllDeviceProxies;
+			TargetDeviceServicesModule->GetDeviceProxyManager()->GetAllProxies(PlatformName, AllDeviceProxies);
 
+			// Remove any Simulator based devices
+			TArray<TSharedPtr<ITargetDeviceProxy>> DeviceProxies;
+			DeviceProxies = AllDeviceProxies.FilterByPredicate([](const TSharedPtr<ITargetDeviceProxy>& Device)
+			{
+				return Device->GetConnectionType() != TEXT("Simulator");
+			});
+			
 			if (DeviceProxies.Num() > 0)
 			{
 				FString LastChosen;
@@ -1781,7 +1810,7 @@ void FTurnkeySupportModule::MakeQuickLaunchItems(class UToolMenu* Menu, FOnQuick
 				FText Tooltip;
 
 				GenerateDeviceProxyMenuParams(DeviceProxies[0], PlatformName, Action, Tooltip, ExternalOnClickDelegate);
-				
+
  				const ITargetPlatform* Platform = GetTargetPlatformManager()->FindTargetPlatform(PlatformName.ToString());
 				bool bGroupDevices = Platform->SupportsFeature(ETargetPlatformFeatures::ShowAsPlatformGroup);
 				if (DeviceProxies.Num() == 1 && !bGroupDevices)
@@ -1805,8 +1834,16 @@ void FTurnkeySupportModule::MakeQuickLaunchItems(class UToolMenu* Menu, FOnQuick
 								FToolMenuSection& Section = SubToolMenu->AddSection(NAME_None);
 
 								// re-get the proxies, just in case they changed
+								TArray<TSharedPtr<ITargetDeviceProxy>> AllDeviceProxies;
+								TargetDeviceServicesModule->GetDeviceProxyManager()->GetAllProxies(PlatformName, AllDeviceProxies);
+
+								// Remove any Simulator based devices
 								TArray<TSharedPtr<ITargetDeviceProxy>> DeviceProxies;
-								TargetDeviceServicesModule->GetDeviceProxyManager()->GetAllProxies(PlatformName, DeviceProxies);
+								DeviceProxies = AllDeviceProxies.FilterByPredicate([](const TSharedPtr<ITargetDeviceProxy> &Device)
+													{ 
+														return Device->GetConnectionType() != TEXT("Simulator"); 
+													});
+
 								// for each one, put an entry (even the one that was in the outer menu, for less confusion)
 								for (const TSharedPtr<ITargetDeviceProxy>& Proxy : DeviceProxies)
 								{
@@ -1909,6 +1946,150 @@ void FTurnkeySupportModule::MakeQuickLaunchItems(class UToolMenu* Menu, FOnQuick
 
 }
 
+static void MakeSimulatorMenu(const FString& ModelType, const TArray<TSharedPtr<ITargetDeviceProxy>>& DeviceProxies, FOnQuickLaunchSelected ExternalOnClickDelegate, FToolMenuSection& DynamicSection)
+{
+	FName PlatformName = TEXT("IOS");
+	FName SubPlatformName = *ModelType;
+	FString LastChosen;
+	const TSharedPtr<ITargetDeviceProxy>* LastChosenProxy;
+	if (PerPlatformLastSimChosen.Contains(SubPlatformName))
+	{
+		LastChosen = PerPlatformLastSimChosen[SubPlatformName];
+		const TSharedPtr<ITargetDeviceProxy>* Proxy = DeviceProxies.FindByPredicate([LastChosen](const TSharedPtr<ITargetDeviceProxy>& InProxy)
+		{
+			return InProxy->GetName() == LastChosen;
+		});
+		if (Proxy)
+		{
+			LastChosenProxy = Proxy;
+		}
+		else
+		{
+			return;
+		}
+	}
+	else
+	{
+		const TSharedPtr<ITargetDeviceProxy>* FirstProxy = DeviceProxies.FindByPredicate([ModelType](const TSharedPtr<ITargetDeviceProxy>& InProxy)
+		{
+			return InProxy->GetType() == ModelType;
+		});
+		if (FirstProxy)
+		{
+			LastChosen = (*FirstProxy)->GetName();
+			LastChosenProxy = FirstProxy;
+			PerPlatformLastSimChosen.Add(SubPlatformName, LastChosen);
+		}
+		else
+		{
+			// There is no simulator of this model type to display, so return
+			return;
+		}
+	}
+	
+	// always use the first one, after sorting
+	FUIAction Action;
+	FText Tooltip;
+	
+	ITargetDeviceServicesModule* TargetDeviceServicesModule = static_cast<ITargetDeviceServicesModule*>(FModuleManager::Get().LoadModule(TEXT("TargetDeviceServices")));
+
+	GenerateDeviceProxyMenuParams(*LastChosenProxy, PlatformName, Action, Tooltip, ExternalOnClickDelegate);
+
+	const ITargetPlatform* Platform = GetTargetPlatformManager()->FindTargetPlatform(PlatformName.ToString());
+	DynamicSection.AddSubMenu(
+		NAME_None,
+		MakeSdkStatusAttribute(PlatformName, *LastChosenProxy),
+		Tooltip,
+		FNewToolMenuDelegate::CreateLambda([TargetDeviceServicesModule, ModelType, PlatformName, LastChosen, DeviceProxies, ExternalOnClickDelegate](UToolMenu* SubToolMenu)
+		{
+			FToolMenuSection& Section = SubToolMenu->AddSection(NAME_None);
+			// for each one, put an entry 
+			for (const TSharedPtr<ITargetDeviceProxy>& Proxy : DeviceProxies)
+			{
+				if (Proxy->GetType() != ModelType)
+				{
+					continue;
+				}
+				
+				// Skip over the top level menu item
+				if (LastChosen == Proxy->GetName())
+				{
+					continue;
+				}
+				
+				FUIAction SubAction;
+				FText SubTooltip;
+				GenerateDeviceProxyMenuParams(Proxy, PlatformName, SubAction, SubTooltip, ExternalOnClickDelegate);
+				Section.AddMenuEntry(NAME_None,
+									 MakeSdkStatusAttribute(PlatformName, Proxy),
+									 SubTooltip,
+									 MakePlatformSdkIconAttribute(PlatformName, Proxy),
+									 SubAction,
+									 EUserInterfaceActionType::Button);
+			}
+		}),
+		Action,
+		EUserInterfaceActionType::Check,
+		false,
+		MakePlatformSdkIconAttribute(PlatformName, *LastChosenProxy),
+		true
+	);
+	
+}
+
+void FTurnkeySupportModule::MakeSimulatorItems(class UToolMenu* Menu, FOnQuickLaunchSelected ExternalOnClickDelegate) const
+{
+	FToolMenuSection& MenuSection = Menu->AddSection("SimulatorDevices", LOCTEXT("Simulators", "Simulators"));
+
+	MenuSection.AddDynamicEntry("PlatformsMenu", FNewToolMenuSectionDelegate::CreateLambda([ExternalOnClickDelegate](FToolMenuSection& DynamicSection)
+	{
+		TArray<FString> DeviceIdsToQuery;
+		ITargetDeviceServicesModule* TargetDeviceServicesModule = static_cast<ITargetDeviceServicesModule*>(FModuleManager::Get().LoadModule(TEXT("TargetDeviceServices")));
+		for (const auto& Pair : FDataDrivenPlatformInfoRegistry::GetAllPlatformInfos())
+		{
+			if (Pair.Value.bIsFakePlatform)
+			{
+				continue;
+			}
+
+			FName PlatformName = Pair.Key;
+			const FDataDrivenPlatformInfo& Info = Pair.Value;
+			if (FDataDrivenPlatformInfoRegistry::IsPlatformHiddenFromUI(PlatformName))
+			{
+				continue;
+			}
+
+			// look for devices for all platforms, even if the platform isn't installed - Turnkey can install Sdk after selecting LaunchOn
+			TArray<TSharedPtr<ITargetDeviceProxy>> AllDeviceProxies;
+			TargetDeviceServicesModule->GetDeviceProxyManager()->GetAllProxies(PlatformName, AllDeviceProxies);
+
+			// We only care about Simulator based devices
+			TArray<TSharedPtr<ITargetDeviceProxy>> DeviceProxies;
+			DeviceProxies = AllDeviceProxies.FilterByPredicate([](const TSharedPtr<ITargetDeviceProxy>& Device)
+			{
+				return Device->GetConnectionType() == TEXT("Simulator");
+			});
+			
+			if (DeviceProxies.Num() > 0)
+			{
+				// Partition based on device type
+				Algo::Sort(DeviceProxies, [](TSharedPtr<ITargetDeviceProxy> A, TSharedPtr<ITargetDeviceProxy> B)
+						   { return A->GetModel() >= B->GetModel(); });
+				
+				MakeSimulatorMenu(TEXT("Phone"), DeviceProxies, ExternalOnClickDelegate, DynamicSection);
+				MakeSimulatorMenu(TEXT("Tablet"), DeviceProxies, ExternalOnClickDelegate, DynamicSection);
+			}
+		}
+
+		// now kick-off any devices that need to be updated
+		if (DeviceIdsToQuery.Num() > 0)
+		{
+			ITurnkeySupportModule::Get().UpdateSdkInfoForDevices(DeviceIdsToQuery);
+		}
+	}
+	));
+}
+
 TSharedRef<SWidget> FTurnkeySupportModule::MakeTurnkeyMenuWidget() const
 {
 	FTurnkeySupportCommands::Register();
@@ -1924,6 +2105,7 @@ TSharedRef<SWidget> FTurnkeySupportModule::MakeTurnkeyMenuWidget() const
 
 		FOnQuickLaunchSelected EmptyFunc;
 		MakeQuickLaunchItems(Menu, EmptyFunc);
+		MakeSimulatorItems(Menu, EmptyFunc);
 
 		// need to make this dyamic so icons, etc can update with SDK 
 		// shared devices section
@@ -2607,8 +2789,15 @@ void FTurnkeySupportModule::RepeatQuickLaunch(FString DeviceId)
 	TSharedPtr<ITargetDeviceProxy> Proxy = TargetDeviceServicesModule->GetDeviceProxyManager()->FindProxyDeviceForTargetDevice(DeviceId);
 
 	if (Proxy.IsValid())
-	{
-		HandleLaunchOnDeviceActionExecute(DeviceId, Proxy->GetName(), true);
+	{				
+		if (Proxy->IsSimulated())
+		{
+			HandleLaunchOnSimulatorActionExecute(DeviceId, Proxy->GetName(), true);
+		}
+		else
+		{
+			HandleLaunchOnDeviceActionExecute(DeviceId, Proxy->GetName(), true);
+		}
 	}
 	else
 	{
@@ -2688,4 +2877,4 @@ IMPLEMENT_MODULE(FTurnkeySupportModule, TurnkeySupport);
 
 #undef LOCTEXT_NAMESPACE
 
-#endif // UE_WITH_TURNKEY_SUPPORT
\ No newline at end of file
+#endif // UE_WITH_TURNKEY_SUPPORT
diff --git a/Engine/Source/Editor/TurnkeySupport/Private/TurnkeySupportModule.h b/Engine/Source/Editor/TurnkeySupport/Private/TurnkeySupportModule.h
--- a/Engine/Source/Editor/TurnkeySupport/Private/TurnkeySupportModule.h
+++ b/Engine/Source/Editor/TurnkeySupport/Private/TurnkeySupportModule.h
@@ -23,6 +23,7 @@ public:
 
 	virtual void MakeTurnkeyMenu(struct FToolMenuSection& MenuSection) const override;
 	virtual void MakeQuickLaunchItems(class UToolMenu* Menu, FOnQuickLaunchSelected ExternalOnClickDelegate) const override;
+	virtual void MakeSimulatorItems(class UToolMenu* Menu, FOnQuickLaunchSelected ExternalOnClickDelegate) const override;
 	virtual void RepeatQuickLaunch(FString DeviceId) override;
 
 
diff --git a/Engine/Source/Editor/TurnkeySupport/Public/Interfaces/ITurnkeySupportModule.h b/Engine/Source/Editor/TurnkeySupport/Public/Interfaces/ITurnkeySupportModule.h
--- a/Engine/Source/Editor/TurnkeySupport/Public/Interfaces/ITurnkeySupportModule.h
+++ b/Engine/Source/Editor/TurnkeySupport/Public/Interfaces/ITurnkeySupportModule.h
@@ -77,6 +77,11 @@ public:
 	 */
 	virtual void MakeQuickLaunchItems(class UToolMenu* Menu, FOnQuickLaunchSelected ExternalOnClickDelegate) const = 0;
 
+	/**
+	 * Make menu items for the Simulator submenu items, so they can be added to the Play menu
+	 */
+	virtual void MakeSimulatorItems(class UToolMenu* Menu, FOnQuickLaunchSelected ExternalOnClickDelegate) const = 0;
+
 	/**
 	 * Runs Turnkey to get the Sdk information for all known platforms
 	 */
