Launch build onto iOS Simulator
Part 3 of 3

diff --git a/Engine/Source/Developer/LauncherServices/Private/Launcher/LauncherWorker.cpp b/Engine/Source/Developer/LauncherServices/Private/Launcher/LauncherWorker.cpp
--- a/Engine/Source/Developer/LauncherServices/Private/Launcher/LauncherWorker.cpp
+++ b/Engine/Source/Developer/LauncherServices/Private/Launcher/LauncherWorker.cpp
@@ -493,6 +493,15 @@ FString FLauncherWorker::CreateUATCommand( const ILauncherProfileRef& InProfile,
 		UATCommand += TEXT(" -target=") + InProfile->GetBuildTarget();
 	}
 
+	if (InProfile->IsDeviceASimulator())
+	{
+		if (Platforms.Contains(TEXT("IOS")))
+		{
+			UATCommand += TEXT(" -clientarchitecture=iossimulator");
+		}
+		// TODO: add tvOS and VisionOS simulators below
+	}
+
 	// device list
 	FString DeviceNames = TEXT("");
 	FString DeviceCommand = TEXT("");
diff --git a/Engine/Source/Developer/LauncherServices/Private/Profiles/LauncherProfile.h b/Engine/Source/Developer/LauncherServices/Private/Profiles/LauncherProfile.h
--- a/Engine/Source/Developer/LauncherServices/Private/Profiles/LauncherProfile.h
+++ b/Engine/Source/Developer/LauncherServices/Private/Profiles/LauncherProfile.h
@@ -2101,6 +2101,7 @@ public:
 		bUseIoStore = false;
 		bUseZenStore = false;
 		bShouldUpdateFlash = false;
+		bIsDeviceASimulator = false;
 		bMakeBinaryConfig = false;
 
 		RefreshValidBuildTargets();
@@ -2632,6 +2633,19 @@ public:
 		return bShouldUpdateFlash;
 	}
 
+	virtual void SetDeviceIsASimulator(bool bInIsDeviceASimualtor) override
+	{
+		bIsDeviceASimulator = bInIsDeviceASimualtor;
+	}
+
+	/**
+	 * Is the Launch device actually a simulator?
+	 */
+	virtual bool IsDeviceASimulator() const override
+	{
+		return bIsDeviceASimulator;
+	}
+
 	virtual void SetMakeBinaryConfig(bool bInMakeBinaryConfig) override
 	{
 		bMakeBinaryConfig = bInMakeBinaryConfig;
@@ -3253,6 +3267,9 @@ private:
 	// Update flash on device before running
 	bool bShouldUpdateFlash;
 
+	// Is the launch device actually a Simulator
+	bool bIsDeviceASimulator;
+
 private:
 
 	// cook in the editor callbacks (not valid for any other cook mode)
diff --git a/Engine/Source/Developer/LauncherServices/Public/ILauncherProfile.h b/Engine/Source/Developer/LauncherServices/Public/ILauncherProfile.h
--- a/Engine/Source/Developer/LauncherServices/Public/ILauncherProfile.h
+++ b/Engine/Source/Developer/LauncherServices/Public/ILauncherProfile.h
@@ -1390,6 +1390,16 @@ public:
 	 */
 	virtual bool ShouldUpdateDeviceFlash() const = 0;
 
+	/**
+	 * Sets whether or not the Device is a Simulator
+	 */
+	virtual void SetDeviceIsASimulator(bool bInIsDeviceASimualtor) = 0;
+
+	/**
+	 * Whether or not the Device is a Simulator
+	 */
+	virtual bool IsDeviceASimulator() const = 0;
+
 public:
 	/**
 	 * Helper function to get all of the build targets available for this profile, based on its current project & cook platforms
diff --git a/Engine/Source/Editor/TurnkeySupport/Private/TurnkeyEditorSupport.cpp b/Engine/Source/Editor/TurnkeySupport/Private/TurnkeyEditorSupport.cpp
--- a/Engine/Source/Editor/TurnkeySupport/Private/TurnkeyEditorSupport.cpp
+++ b/Engine/Source/Editor/TurnkeySupport/Private/TurnkeyEditorSupport.cpp
@@ -90,7 +90,7 @@ void FTurnkeyEditorSupport::PrepareToLaunchRunningMap(const FString& DeviceId, c
 #endif
 }
 
-void FTurnkeyEditorSupport::LaunchRunningMap(const FString& DeviceId, const FString& DeviceName, const FString& ProjectPath, bool bUseTurnkey)
+void FTurnkeyEditorSupport::LaunchRunningMap(const FString& DeviceId, const FString& DeviceName, const FString& ProjectPath, bool bUseTurnkey, bool bOnSimulator)
 {
 #if WITH_EDITOR
 	FTargetDeviceId TargetDeviceId;
@@ -130,6 +130,7 @@ void FTurnkeyEditorSupport::LaunchRunningMap(const FString& DeviceId, const FStr
 			// @todo turnkey: we set this to false because we will kick off a Turnkey run before cooking, etc, to get an early warning. however, if it's too difficult
 			// to get an error back from CreateUatTask, then we should set this to bUseTurnkey and remove the block below, and let the code in FLauncherWorker::CreateAndExecuteTasks handle it
 			DeviceInfo.bUpdateDeviceFlash = false;
+			DeviceInfo.bIsSimulator = bOnSimulator;
 
 			FRequestPlaySessionParams SessionParams;
 			SessionParams.SessionDestination = EPlaySessionDestinationType::Launcher;
@@ -293,4 +294,4 @@ bool FTurnkeyEditorSupport::IsPIERunning()
 
 #undef LOCTEXT_NAMESPACE
 
-#endif // UE_WITH_TURNKEY_SUPPORT
\ No newline at end of file
+#endif // UE_WITH_TURNKEY_SUPPORT
diff --git a/Engine/Source/Editor/TurnkeySupport/Private/TurnkeyEditorSupport.h b/Engine/Source/Editor/TurnkeySupport/Private/TurnkeyEditorSupport.h
--- a/Engine/Source/Editor/TurnkeySupport/Private/TurnkeyEditorSupport.h
+++ b/Engine/Source/Editor/TurnkeySupport/Private/TurnkeyEditorSupport.h
@@ -14,7 +14,7 @@ public:
 	static FString GetUATOptions();
 
 	static void PrepareToLaunchRunningMap(const FString& DeviceId, const FString& DeviceName);
-	static void LaunchRunningMap(const FString& DeviceId, const FString& DeviceName, const FString& ProjectPath, bool bUseTurnkey);
+	static void LaunchRunningMap(const FString& DeviceId, const FString& DeviceName, const FString& ProjectPath, bool bUseTurnkey, bool bOnSimulator = false);
 	static void AddEditorOptions(struct FToolMenuSection& MenuBuilder);
 	
 	static void SaveAll();
diff --git a/Engine/Source/Editor/TurnkeySupport/Private/TurnkeySupportModule.cpp b/Engine/Source/Editor/TurnkeySupport/Private/TurnkeySupportModule.cpp
--- a/Engine/Source/Editor/TurnkeySupport/Private/TurnkeySupportModule.cpp
+++ b/Engine/Source/Editor/TurnkeySupport/Private/TurnkeySupportModule.cpp
@@ -1619,7 +1619,7 @@ static void LaunchOnDevice(const FString& DeviceId, const FString& DeviceName, b
 
 static void LaunchOnSimulator(const FString& DeviceId, const FString& DeviceName, bool bUseTurnkey)
 {								
-	// FTurnkeyEditorSupport::LaunchRunningMap(DeviceId, DeviceName, GetProjectPathForTurnkey(), false, true);
+	FTurnkeyEditorSupport::LaunchRunningMap(DeviceId, DeviceName, GetProjectPathForTurnkey(), false, true);
 }
 
 static void PrepareLaunchOn(FString DeviceId, FString DeviceName)
diff --git a/Engine/Source/Editor/UnrealEd/Private/PlayLevelViaLauncher.cpp b/Engine/Source/Editor/UnrealEd/Private/PlayLevelViaLauncher.cpp
--- a/Engine/Source/Editor/UnrealEd/Private/PlayLevelViaLauncher.cpp
+++ b/Engine/Source/Editor/UnrealEd/Private/PlayLevelViaLauncher.cpp
@@ -221,6 +221,8 @@ void UEditorEngine::StartPlayUsingLauncherSession(FRequestPlaySessionParams& InR
 
 	LauncherProfile->AddCookedPlatform(LaunchPlatformName);
 
+	LauncherProfile->SetDeviceIsASimulator(InRequestParams.LauncherTargetDevice->bIsSimulator);
+
 	// select the quickest cook mode based on which in editor cook mode is enabled
 	const UCookerSettings& CookerSettings = *GetDefault<UCookerSettings>();
 	const UEditorExperimentalSettings& ExperimentalSettings = *GetDefault<UEditorExperimentalSettings>();
diff --git a/Engine/Source/Editor/UnrealEd/Public/PlayInEditorDataTypes.h b/Engine/Source/Editor/UnrealEd/Public/PlayInEditorDataTypes.h
--- a/Engine/Source/Editor/UnrealEd/Public/PlayInEditorDataTypes.h
+++ b/Engine/Source/Editor/UnrealEd/Public/PlayInEditorDataTypes.h
@@ -143,6 +143,8 @@ public:
 		FString DeviceName;
 		/** If True, a remote play session will attempt to update the flash/software on the target device if it's out of date */
 		bool bUpdateDeviceFlash = false;
+		/** If True, the launch device is a Simulator */
+		bool bIsSimulator = false;
 	};
 
 	/** Where should the session be launched? May be local or remote. */
diff --git a/Engine/Source/Programs/AutomationTool/IOS/IOSPlatform.Automation.cs b/Engine/Source/Programs/AutomationTool/IOS/IOSPlatform.Automation.cs
--- a/Engine/Source/Programs/AutomationTool/IOS/IOSPlatform.Automation.cs
+++ b/Engine/Source/Programs/AutomationTool/IOS/IOSPlatform.Automation.cs
@@ -1915,6 +1915,16 @@ public class IOSPlatform : ApplePlatform
 		string LibimobileDeviceArguments =  "-u " + Params.DeviceNames[0] + " -i " + "\"" +  Path.GetFullPath(AppToDeploy) + "\"";
 		LibimobileDeviceArguments = GetLibimobileDeviceNetworkedArgument(LibimobileDeviceArguments, Params.DeviceNames[0]);
 
+		// If we deploying to a Simulator, use "xcrun simctl" instead
+		{
+			TargetReceipt Targets = SC.StageTargets[0].Receipt;
+			if (Targets.Architectures.SingleArchitecture == UnrealArch.IOSSimulator)
+			{
+				DeviceInstaller = "xcrun";
+				LibimobileDeviceArguments = "simctl install " + Params.DeviceNames[0] + " \"" + Path.GetFullPath(AppToDeploy) + "\"";
+			}
+		}
+
 		// check for it in the stage directory
 		string CurrentDir = Directory.GetCurrentDirectory();
 		Directory.SetCurrentDirectory(CombinePaths(CmdEnv.LocalRoot, "Engine/Binaries/DotNET/IOS/"));
