- Fixed an issue where per-platform packaging settings were getting confused when project settings were changed (the per-platform ones were not getting updated). This is a heavy handed approach to wiping platform config caches, and reloading them as needed

diff --git a/Engine/Source/Developer/DeveloperToolSettings/Private/ProjectPackagingSettings.cpp b/Engine/Source/Developer/DeveloperToolSettings/Private/ProjectPackagingSettings.cpp
--- a/Engine/Source/Developer/DeveloperToolSettings/Private/ProjectPackagingSettings.cpp
+++ b/Engine/Source/Developer/DeveloperToolSettings/Private/ProjectPackagingSettings.cpp
@@ -88,6 +88,15 @@ void UProjectPackagingSettings::PostEditChangeProperty(FPropertyChangedEvent& Pr
 		? PropertyChangedEvent.MemberProperty->GetFName()
 		: NAME_None;
 
+	// there is an issue where the default BuildConfig doesn't get propagated to the per-platform packaging settings instances
+	// because their ConfigSystem doesn't get updated - so we do the nuclear option here and wipe them all - they will get re-loaded
+	// when queried next. we just do it for all properties because any one of these could need the per-platform version to update
+	// this code runs before the ReloadConfig is called on each instance
+	if (Name != NAME_None)
+	{
+		FConfigCacheIni::ClearOtherPlatformConfigs();
+	}
+
 	if (Name == FName(TEXT("DirectoriesToAlwaysCook")) || Name == FName(TEXT("DirectoriesToNeverCook")) || Name == FName(TEXT("TestDirectoriesToNotSearch")) || Name == NAME_None)
 	{
 		// We need to fix paths for no name updates to catch the reloadconfig call
diff --git a/Engine/Source/Editor/TurnkeySupport/Private/TurnkeySupportModule.cpp b/Engine/Source/Editor/TurnkeySupport/Private/TurnkeySupportModule.cpp
--- a/Engine/Source/Editor/TurnkeySupport/Private/TurnkeySupportModule.cpp
+++ b/Engine/Source/Editor/TurnkeySupport/Private/TurnkeySupportModule.cpp
@@ -267,6 +267,9 @@ public:
 			PackagingSettings->LoadSettingsForPlatform(PlatformString);
 		}
 
+		// make sure any flushed settings are reloaded
+		PackagingSettings->ReloadConfig();
+
 		return PackagingSettings;
 	}
 
@@ -1258,7 +1261,14 @@ static void MakeTurnkeyPlatformMenu(UToolMenu* ToolMenu, FName IniPlatformName,
 
 		FToolMenuSection& ConfigSection = ToolMenu->AddSection("BuildConfig", LOCTEXT("TurnkeySection_BuildConfig", "Binary Configuration"));
 		
-		const UProjectPackagingSettings::FConfigurationInfo& ConfigInfo = UProjectPackagingSettings::ConfigurationInfo[static_cast<int32>(AllPlatformPackagingSettings->BuildConfiguration)];
+		EProjectPackagingBuildConfigurations BuildConfig = GetMutableDefault<UPlatformsMenuSettings>()->GetBuildConfigurationForPlatform(IniPlatformName);
+		// if PPBC_MAX is set, then the project default should be used instead of the per platform build config
+		if (BuildConfig == EProjectPackagingBuildConfigurations::PPBC_MAX)
+		{
+			BuildConfig = PackagingSettings->BuildConfiguration;
+		}
+
+		const UProjectPackagingSettings::FConfigurationInfo& ConfigInfo = UProjectPackagingSettings::ConfigurationInfo[static_cast<int32>(BuildConfig)];
 		ConfigSection.AddMenuEntry(
 			NAME_None,
 			FText::Format(LOCTEXT("DefaultConfiguration",  "Use Project Setting ({0})"), ConfigInfo.Name),
