Show a notification that the editor will need to be rebuilt when installing an SDK with no compiled platform support.

 - default "editor needs restarting" will not be enough to pick up the new platform.

#rnx

diff --git a/Engine/Source/Editor/TurnkeySupport/Private/TurnkeyEditorSupport.cpp b/Engine/Source/Editor/TurnkeySupport/Private/TurnkeyEditorSupport.cpp
--- a/Engine/Source/Editor/TurnkeySupport/Private/TurnkeyEditorSupport.cpp
+++ b/Engine/Source/Editor/TurnkeySupport/Private/TurnkeyEditorSupport.cpp
@@ -11,6 +11,8 @@
 #if WITH_EDITOR
 #include "UnrealEdMisc.h"
 #include "Framework/MultiBox/MultiBoxBuilder.h"
+#include "Framework/Notifications/NotificationManager.h"
+#include "Widgets/Notifications/SNotificationList.h"
 #include "Editor/UnrealEdEngine.h"
 #include "UnrealEdGlobals.h"
 #include "PlatformInfo.h"
@@ -262,6 +264,18 @@ void FTurnkeyEditorSupport::ShowRestartToast()
 #endif
 }
 
+void FTurnkeyEditorSupport::ShowRebuildToast()
+{
+#if WITH_EDITOR
+	// show a dialog asking for the editor to be rebuilt
+	// @todo: if the source code is available should display a 'rebuild now' button
+	FNotificationInfo* Info = new FNotificationInfo( LOCTEXT("Turnkey_RebuildRequired", "The editor must be rebuilt to include support for the new platform") );
+	Info->ExpireDuration = 8.0f;
+	Info->Image = FAppStyle::GetBrush("Icons.InfoWithColor");
+	FSlateNotificationManager::Get().QueueNotification(Info);
+#endif
+}
+
 bool FTurnkeyEditorSupport::CheckSupportedPlatforms(FName IniPlatformName)
 {
 #if WITH_EDITOR
diff --git a/Engine/Source/Editor/TurnkeySupport/Private/TurnkeyEditorSupport.h b/Engine/Source/Editor/TurnkeySupport/Private/TurnkeyEditorSupport.h
--- a/Engine/Source/Editor/TurnkeySupport/Private/TurnkeyEditorSupport.h
+++ b/Engine/Source/Editor/TurnkeySupport/Private/TurnkeyEditorSupport.h
@@ -23,6 +23,7 @@ public:
 
 	static bool ShowOKCancelDialog(FText Message, FText Title);
 	static void ShowRestartToast();
+	static void ShowRebuildToast();
 	static bool CheckSupportedPlatforms(FName IniPlatformName);
 	static void ShowInstallationHelp(FName IniPlatformName, FString DocLink);
 	static bool IsPIERunning();
diff --git a/Engine/Source/Editor/TurnkeySupport/Private/TurnkeySupportModule.cpp b/Engine/Source/Editor/TurnkeySupport/Private/TurnkeySupportModule.cpp
--- a/Engine/Source/Editor/TurnkeySupport/Private/TurnkeySupportModule.cpp
+++ b/Engine/Source/Editor/TurnkeySupport/Private/TurnkeySupportModule.cpp
@@ -936,7 +936,14 @@ static void TurnkeyInstallSdk(FString IniPlatformName, bool bPreferFull, bool bF
 			RenderUtilsInit();
 #endif
 
-			FTurnkeyEditorSupport::ShowRestartToast();
+			if (FDataDrivenPlatformInfoRegistry::HasCompiledSupportForPlatform(*IniPlatformName, FDataDrivenPlatformInfoRegistry::EPlatformNameType::Ini))
+			{
+				FTurnkeyEditorSupport::ShowRestartToast();
+			}
+			else
+			{
+				FTurnkeyEditorSupport::ShowRebuildToast();
+			}
 		});
 	}
 	);
