Return Android devices with Turnkey, and Editor `Platforms` menu update

The goal of this update is to have valid Android devices not showing with a Warning in the Platforms menu.

To allow this, the AndroidPlatformSDK now return the list of connected devices and their SDK Versions.
The change in TurnkeySupportModule ensures we are stripping the cook flavor from the DeviceId which would fail to verify the SDK.

Regarding the Editor `Platforms` menu, the code was adjusted to have devices listed outside of a submenu when `Platform->SupportsFeature(ETargetPlatformFeatures::ShowAsPlatformGroup)` returns false.
Another addition was a new `ETargetPlatformFeatures::ShowAggregateDevice` to allow platforms to not list the aggregate device (ex: `All_Android_on...`)

diff --git a/Engine/Source/Developer/Android/AndroidTargetPlatformSettings/Private/AndroidTargetPlatformSettings.cpp b/Engine/Source/Developer/Android/AndroidTargetPlatformSettings/Private/AndroidTargetPlatformSettings.cpp
--- a/Engine/Source/Developer/Android/AndroidTargetPlatformSettings/Private/AndroidTargetPlatformSettings.cpp
+++ b/Engine/Source/Developer/Android/AndroidTargetPlatformSettings/Private/AndroidTargetPlatformSettings.cpp
@@ -62,6 +62,9 @@ bool FAndroidTargetPlatformSettings::SupportsFeature(ETargetPlatformFeatures Fea
 	case ETargetPlatformFeatures::DeviceOutputLog:
 		return true;
 
+	case ETargetPlatformFeatures::ShowAggregateDevice:
+		return false;
+	
 	case ETargetPlatformFeatures::LowQualityLightmaps:
 	case ETargetPlatformFeatures::MobileRendering:
 		return SupportsES31() || SupportsVulkan();
diff --git a/Engine/Source/Developer/IOS/IOSTargetPlatformSettings/Private/IOSTargetPlatformSettings.cpp b/Engine/Source/Developer/IOS/IOSTargetPlatformSettings/Private/IOSTargetPlatformSettings.cpp
--- a/Engine/Source/Developer/IOS/IOSTargetPlatformSettings/Private/IOSTargetPlatformSettings.cpp
+++ b/Engine/Source/Developer/IOS/IOSTargetPlatformSettings/Private/IOSTargetPlatformSettings.cpp
@@ -99,9 +99,6 @@ bool FIOSTargetPlatformSettings::SupportsFeature( ETargetPlatformFeatures Featur
 			return bUsesARMCompressor;
 		}
 
-		case ETargetPlatformFeatures::ShowAsPlatformGroup:
-			return false;
-
 		case ETargetPlatformFeatures::SupportsMultipleConnectionTypes:
 			return true;
 
diff --git a/Engine/Source/Developer/TargetPlatform/Public/Common/TargetPlatformSettingsBase.h b/Engine/Source/Developer/TargetPlatform/Public/Common/TargetPlatformSettingsBase.h
--- a/Engine/Source/Developer/TargetPlatform/Public/Common/TargetPlatformSettingsBase.h
+++ b/Engine/Source/Developer/TargetPlatform/Public/Common/TargetPlatformSettingsBase.h
@@ -114,6 +114,9 @@ public:
 		case ETargetPlatformFeatures::CanCookPackages:
 			return false;
 
+		case ETargetPlatformFeatures::ShowAggregateDevice:
+			return true;
+
 		case ETargetPlatformFeatures::TextureStreaming:
 			return TPlatformProperties::SupportsTextureStreaming();
 		case ETargetPlatformFeatures::MeshLODStreaming:
diff --git a/Engine/Source/Developer/TargetPlatform/Public/Interfaces/ITargetPlatformSettings.h b/Engine/Source/Developer/TargetPlatform/Public/Interfaces/ITargetPlatformSettings.h
--- a/Engine/Source/Developer/TargetPlatform/Public/Interfaces/ITargetPlatformSettings.h
+++ b/Engine/Source/Developer/TargetPlatform/Public/Interfaces/ITargetPlatformSettings.h
@@ -106,8 +106,8 @@ enum class ETargetPlatformFeatures
 	/** The platform supports Luminance + Alpha encoding mode for normalmaps */
 	NormalmapLAEncodingMode,
 
-	/** All devices of this platform should be grouped under one platform group */
-	ShowAsPlatformGroup,
+	/** Should the Aggregate Device be shown for this platform (ex: All_Android_on_...) */
+	ShowAggregateDevice,
 
 	/** Does the platform allow various connection types to be used (ie: wifi and usb) */
 	SupportsMultipleConnectionTypes,
diff --git a/Engine/Source/Editor/TurnkeySupport/Private/TurnkeySupportModule.cpp b/Engine/Source/Editor/TurnkeySupport/Private/TurnkeySupportModule.cpp
--- a/Engine/Source/Editor/TurnkeySupport/Private/TurnkeySupportModule.cpp
+++ b/Engine/Source/Editor/TurnkeySupport/Private/TurnkeySupportModule.cpp
@@ -974,7 +974,7 @@ static TAttribute<FText> MakeSdkStatusAttribute(FName IniPlatformName, TSharedPt
 
 static FSlateIcon MakePlatformSdkIconAttribute(FName IniPlatformName, TSharedPtr< ITargetDeviceProxy> DeviceProxy)
 {
-	FString DeviceId = DeviceProxy ? DeviceProxy->GetTargetDeviceId(NAME_None) : FString();
+	FString DeviceId = DeviceProxy ? DeviceProxy->GetTargetDeviceId(DeviceProxy->GetVanillaPlatformId(NAME_None)) : FString();
 
 	// get the status, or Unknown if it's not there
 	ETurnkeyPlatformSdkStatus Status = DeviceId.Len() ? ITurnkeySupportModule::Get().GetSdkInfoForDeviceId(DeviceId).Status : ITurnkeySupportModule::Get().GetSdkInfo(IniPlatformName, false).Status;
@@ -1797,12 +1797,22 @@ void FTurnkeySupportModule::MakeQuickLaunchItems(class UToolMenu* Menu, FOnQuick
 			}
 
 			// look for devices for all platforms, even if the platform isn't installed - Turnkey can install Sdk after selecting LaunchOn
-			TArray<TSharedPtr<ITargetDeviceProxy>> AllDeviceProxies;
-			TargetDeviceServicesModule->GetDeviceProxyManager()->GetAllProxies(PlatformName, AllDeviceProxies);
+			TArray<TSharedPtr<ITargetDeviceProxy>> DeviceProxies;
+			TargetDeviceServicesModule->GetDeviceProxyManager()->GetAllProxies(PlatformName, DeviceProxies);
+			
+			const ITargetPlatform* Platform = GetTargetPlatformManager()->FindTargetPlatform(PlatformName.ToString());
+
+			// Remove Aggregate Device
+			if (Platform && !Platform->SupportsFeature(ETargetPlatformFeatures::ShowAggregateDevice))
+			{
+				DeviceProxies = DeviceProxies.FilterByPredicate([](const TSharedPtr<ITargetDeviceProxy>& Device)
+				{
+					return !Device->IsAggregated();
+				});
+			}
 
 			// Remove any Simulator based devices
-			TArray<TSharedPtr<ITargetDeviceProxy>> DeviceProxies;
-			DeviceProxies = AllDeviceProxies.FilterByPredicate([](const TSharedPtr<ITargetDeviceProxy>& Device)
+			DeviceProxies = DeviceProxies.FilterByPredicate([](const TSharedPtr<ITargetDeviceProxy>& Device)
 			{
 				return Device->GetConnectionType() != TEXT("Simulator");
 			});
@@ -1822,16 +1832,13 @@ void FTurnkeySupportModule::MakeQuickLaunchItems(class UToolMenu* Menu, FOnQuick
 				Algo::Sort(DeviceProxies, [&LastChosen = std::as_const(LastChosen)](TSharedPtr<ITargetDeviceProxy> A, TSharedPtr<ITargetDeviceProxy> B)
 						   { return A->GetName() == LastChosen; });
 
-				// always use the first one, after sorting
-				FUIAction Action;
-				FText Tooltip;
-
-				GenerateDeviceProxyMenuParams(DeviceProxies[0], PlatformName, Action, Tooltip, ExternalOnClickDelegate);
 
- 				const ITargetPlatform* Platform = GetTargetPlatformManager()->FindTargetPlatform(PlatformName.ToString());
-				bool bGroupDevices = Platform->SupportsFeature(ETargetPlatformFeatures::ShowAsPlatformGroup);
-				if (DeviceProxies.Num() == 1 && !bGroupDevices)
+				if (DeviceProxies.Num() == 1)
 				{
+					FUIAction Action;
+					FText Tooltip;
+					GenerateDeviceProxyMenuParams(DeviceProxies[0], PlatformName, Action, Tooltip, ExternalOnClickDelegate);
+
 					DynamicSection.AddMenuEntry(
 						NAME_None,
 						MakeSdkStatusAttribute(PlatformName, DeviceProxies[0]),
@@ -1844,35 +1851,43 @@ void FTurnkeySupportModule::MakeQuickLaunchItems(class UToolMenu* Menu, FOnQuick
 				{
 					DynamicSection.AddSubMenu(
 						NAME_None,
-						MakeSdkStatusAttribute(PlatformName, bGroupDevices?nullptr:DeviceProxies[0]),
-						Tooltip,
-						FNewToolMenuDelegate::CreateLambda([TargetDeviceServicesModule, PlatformName, LastChosen, ExternalOnClickDelegate](UToolMenu* SubToolMenu)
+						MakeSdkStatusAttribute(PlatformName, nullptr),
+						FText::FromString(PlatformName.ToString()),
+						FNewToolMenuDelegate::CreateLambda(
+							[TargetDeviceServicesModule, PlatformName, LastChosen, ExternalOnClickDelegate](UToolMenu* SubToolMenu)
 							{
 								FToolMenuSection& Section = SubToolMenu->AddSection(NAME_None);
 
 								// re-get the proxies, just in case they changed
-								TArray<TSharedPtr<ITargetDeviceProxy>> AllDeviceProxies;
-								TargetDeviceServicesModule->GetDeviceProxyManager()->GetAllProxies(PlatformName, AllDeviceProxies);
+								TArray<TSharedPtr<ITargetDeviceProxy>> DeviceProxies;
+								TargetDeviceServicesModule->GetDeviceProxyManager()->GetAllProxies(PlatformName, DeviceProxies);
+								const ITargetPlatform* Platform = GetTargetPlatformManager()->FindTargetPlatform(PlatformName.ToString());
+
+								// Remove Aggregate Device
+								if (Platform && !Platform->SupportsFeature(ETargetPlatformFeatures::ShowAggregateDevice))
+								{
+									DeviceProxies = DeviceProxies.FilterByPredicate([](const TSharedPtr<ITargetDeviceProxy>& Device)
+									{
+										return !Device->IsAggregated();
+									});
+								}
 
 								// Remove any Simulator based devices
-								TArray<TSharedPtr<ITargetDeviceProxy>> DeviceProxies;
-								DeviceProxies = AllDeviceProxies.FilterByPredicate([](const TSharedPtr<ITargetDeviceProxy> &Device)
-													{ 
-														return Device->GetConnectionType() != TEXT("Simulator"); 
-													});
+								DeviceProxies = DeviceProxies.FilterByPredicate([](const TSharedPtr<ITargetDeviceProxy>& Device)
+								{
+									return Device->GetConnectionType() != TEXT("Simulator");
+								});
+								
+								Algo::Sort(DeviceProxies, [&LastChosen = std::as_const(LastChosen)](TSharedPtr<ITargetDeviceProxy> A, TSharedPtr<ITargetDeviceProxy> B)
+									{ return A->GetName() == LastChosen; });
 
 								// for each one, put an entry (even the one that was in the outer menu, for less confusion)
 								for (const TSharedPtr<ITargetDeviceProxy>& Proxy : DeviceProxies)
 								{
-									// Skip over the top level menu item
-									if (LastChosen == Proxy->GetName())
-									{
-										continue;
-									}
-									
 									FUIAction SubAction;
 									FText SubTooltip;
 									GenerateDeviceProxyMenuParams(Proxy, PlatformName, SubAction, SubTooltip, ExternalOnClickDelegate);
+									
 									Section.AddMenuEntry(
 										NAME_None,
 										MakeSdkStatusAttribute(PlatformName, Proxy),
@@ -1883,19 +1898,19 @@ void FTurnkeySupportModule::MakeQuickLaunchItems(class UToolMenu* Menu, FOnQuick
 									);
 								}
 							}),
-						Action,
+						{},
 						EUserInterfaceActionType::Check,
 						false,
-						MakePlatformSdkIconAttribute(PlatformName, bGroupDevices?nullptr:DeviceProxies[0]),
+						MakePlatformSdkIconAttribute(PlatformName, nullptr),
 						true
-						);
+					);
 				}
 
 				ITurnkeySupportModule& TurnkeySupport = ITurnkeySupportModule::Get();
 				// gather any unknown status devices to query at the end
 				for (const TSharedPtr<ITargetDeviceProxy>& Proxy : DeviceProxies)
 				{
-					FString DeviceId = Proxy->GetTargetDeviceId(NAME_None);
+					FString DeviceId = Proxy->GetTargetDeviceId(Proxy->GetVanillaPlatformId(NAME_None));
 					if (TurnkeySupport.GetSdkInfoForDeviceId(DeviceId).Status == ETurnkeyPlatformSdkStatus::Unknown)
 					{
 						DeviceIdsToQuery.Add(DeviceId);
@@ -2679,7 +2694,7 @@ void FTurnkeySupportModule::UpdateSdkInfoForProxy(const TSharedRef<ITargetDevice
 {
 	bool bIsNeeded;
 
-	FString DeviceId = AddedProxy->GetTargetDeviceId(NAME_None);
+	FString DeviceId = AddedProxy->GetTargetDeviceId(AddedProxy->GetVanillaPlatformId(NAME_None));
 	{
 		FScopeLock Lock(&GTurnkeySection);
 		bIsNeeded = GetSdkInfoForDeviceId(DeviceId).Status == ETurnkeyPlatformSdkStatus::Unknown;
diff --git a/Engine/Source/Programs/AutomationTool/Android/AndroidPlatform.Automation.cs b/Engine/Source/Programs/AutomationTool/Android/AndroidPlatform.Automation.cs
--- a/Engine/Source/Programs/AutomationTool/Android/AndroidPlatform.Automation.cs
+++ b/Engine/Source/Programs/AutomationTool/Android/AndroidPlatform.Automation.cs
@@ -502,13 +502,36 @@ public class AndroidPlatform : Platform
 
 		return bHaveAndroidStudio;
 	}
-
-
-
-
-
-
-
+	
+	public override DeviceInfo[] GetDevices()
+	{
+		List<DeviceInfo> Devices = new List<DeviceInfo>();
+		
+		List<string> ConnectedDevices;
+		GetConnectedDevices(null, out ConnectedDevices);
+		
+		foreach (string ConnectedDeviceName in ConnectedDevices)
+		{
+			DeviceInfo CurrentDevice = new DeviceInfo(TargetPlatformType);
+			// GetConnectedDevices returns valid device names with an '@' in front
+			CurrentDevice.Name = ConnectedDeviceName.StartsWith("@") ? ConnectedDeviceName.Remove(0, 1) : ConnectedDeviceName; 
+			
+			CurrentDevice.Id = CurrentDevice.Name;
+			CurrentDevice.bCanConnect = ConnectedDeviceName.StartsWith("@");
+			
+			// Instead, we return the SDK Version, and other parts of the code were adjusted
+			string GetPropCommand = "shell getprop";
+			string SDKVersionCommand = $"{GetPropCommand} ro.build.version.sdk";
+			IProcessResult Result = RunAdbCommand(CurrentDevice.Name, SDKVersionCommand);
+			if (Result.Output.Length > 0)
+			{
+				CurrentDevice.SoftwareVersion = Result.Output.Trim();
+			}
+						
+			Devices.Add(CurrentDevice);
+		}
+		return Devices.ToArray();
+	}
 
 
 
diff --git a/Engine/Source/Programs/Shared/EpicGames.Build/UEBuildPlatformSDK.cs b/Engine/Source/Programs/Shared/EpicGames.Build/UEBuildPlatformSDK.cs
--- a/Engine/Source/Programs/Shared/EpicGames.Build/UEBuildPlatformSDK.cs
+++ b/Engine/Source/Programs/Shared/EpicGames.Build/UEBuildPlatformSDK.cs
@@ -805,7 +805,7 @@ namespace EpicGames.Core
 
 			// convert it to an integer
 			ulong IntVersion;
-			if (!TryConvertVersionToInt(Version, out IntVersion))
+			if (!TryConvertVersionToInt(Version, out IntVersion, VersionInfo.Name))
 			{
 				return false;
 			}
diff --git a/Engine/Source/Programs/UnrealBuildTool/Platform/Android/AndroidPlatformSDK.cs b/Engine/Source/Programs/UnrealBuildTool/Platform/Android/AndroidPlatformSDK.cs
--- a/Engine/Source/Programs/UnrealBuildTool/Platform/Android/AndroidPlatformSDK.cs
+++ b/Engine/Source/Programs/UnrealBuildTool/Platform/Android/AndroidPlatformSDK.cs
@@ -141,6 +141,16 @@ namespace UnrealBuildTool
 			// convert r<num>[letter] to hex
 			if (!String.IsNullOrEmpty(StringValue))
 			{
+				if (String.Compare(Hint, "Software", true) == 0)
+				{
+					// returning a version with a format similar to the ndk version.
+					string VersionString = String.Format("{0}{1:00}{2:00}", StringValue, 0, 0);
+					if (UInt64.TryParse(VersionString, out OutValue))
+					{
+						return true;
+					}
+				}
+				
 				Match Result = Regex.Match(StringValue, @"^r(\d*)([a-z])?");
 
 				if (Result.Success)
