Missing shell quotes around -ProjectFilePath= for the 'Prepare for debugging' flow resulting in Turnkey command failing when ProjectFilePath contains spaces.

[RN] Resolve AutomationTool error when running 'Platforms' -> <plarform name> -> 'Prepare for debugging' for a project for which its path contains spaces.

diff --git a/Engine/Source/Editor/TurnkeySupport/Private/TurnkeySupportModule.cpp b/Engine/Source/Editor/TurnkeySupport/Private/TurnkeySupportModule.cpp
--- a/Engine/Source/Editor/TurnkeySupport/Private/TurnkeySupportModule.cpp
+++ b/Engine/Source/Editor/TurnkeySupport/Private/TurnkeySupportModule.cpp
@@ -526,12 +526,12 @@ public:
 				 GConfig->GetBool(TEXT("/Script/IOSRuntimeSettings.IOSRuntimeSettings"), TEXT("bSupportSecondaryMac"), bSupportSecondaryMac, GEngineIni);
 				 if (bSupportSecondaryMac)
 				 {
-					 FString OtherCommandLine = FString::Printf(TEXT("SetSecondaryRemoteMac -platform=%s -ProjectFilePath=%s"), *IniPlatformName.ToString().ToLower(), *ProjectPath);
+					 FString OtherCommandLine = FString::Printf(TEXT("SetSecondaryRemoteMac -platform=%s -ProjectFilePath=\"%s\""), *IniPlatformName.ToString().ToLower(), *ProjectPath);
 					 FTurnkeyEditorSupport::RunUAT(OtherCommandLine, PlatformInfo->DisplayName, ContentPrepDescription, ContentPrepTaskName, ContentPrepIcon, &AnalyticsParamArray);
 				 }
 				 else
 				 {
-					 FString CommandLine = FString::Printf(TEXT("WrangleContentForDebugging -platform=%s -ProjectFilePath=%s"), *IniPlatformName.ToString().ToLower(), *ProjectPath);
+					 FString CommandLine = FString::Printf(TEXT("WrangleContentForDebugging -platform=%s -ProjectFilePath=\"%s\""), *IniPlatformName.ToString().ToLower(), *ProjectPath);
 					 FTurnkeyEditorSupport::RunUAT(CommandLine, PlatformInfo->DisplayName, ContentPrepDescription, ContentPrepTaskName, ContentPrepIcon, &AnalyticsParamArray);
 				 }
 
@@ -540,7 +540,7 @@ public:
 #elif PLATFORM_MAC
 		 if (IniPlatformName == TEXT("IOS") || IniPlatformName == TEXT("TvOS"))
 		 {
-			 FString CommandLine = FString::Printf(TEXT("WrangleContentForDebugging -platform=%s -ProjectFilePath=%s"), *IniPlatformName.ToString().ToLower(), *ProjectPath);
+			 FString CommandLine = FString::Printf(TEXT("WrangleContentForDebugging -platform=%s -ProjectFilePath=\"%s\""), *IniPlatformName.ToString().ToLower(), *ProjectPath);
 			 FTurnkeyEditorSupport::RunUAT(CommandLine, PlatformInfo->DisplayName, ContentPrepDescription, ContentPrepTaskName, ContentPrepIcon, &AnalyticsParamArray);
 			 return;
 		 }
