fix for project default configuration label

#rnx

diff --git a/Engine/Source/Editor/TurnkeySupport/Private/TurnkeySupportModule.cpp b/Engine/Source/Editor/TurnkeySupport/Private/TurnkeySupportModule.cpp
--- a/Engine/Source/Editor/TurnkeySupport/Private/TurnkeySupportModule.cpp
+++ b/Engine/Source/Editor/TurnkeySupport/Private/TurnkeySupportModule.cpp
@@ -1268,18 +1268,11 @@ static void MakeTurnkeyPlatformMenu(UToolMenu* ToolMenu, FName IniPlatformName,
 
 		FToolMenuSection& ConfigSection = ToolMenu->AddSection("BuildConfig", LOCTEXT("TurnkeySection_BuildConfig", "Binary Configuration"));
 		
-		EProjectPackagingBuildConfigurations BuildConfig = GetMutableDefault<UPlatformsMenuSettings>()->GetBuildConfigurationForPlatform(IniPlatformName);
-		// if PPBC_MAX is set, then the project default should be used instead of the per platform build config
-		if (BuildConfig == EProjectPackagingBuildConfigurations::PPBC_MAX)
-		{
-			BuildConfig = PackagingSettings->BuildConfiguration;
-		}
-
-		const UProjectPackagingSettings::FConfigurationInfo& ConfigInfo = UProjectPackagingSettings::ConfigurationInfo[static_cast<int32>(BuildConfig)];
+		const UProjectPackagingSettings::FConfigurationInfo& DefaultConfigInfo = UProjectPackagingSettings::ConfigurationInfo[static_cast<int32>(PackagingSettings->BuildConfiguration)];
 		ConfigSection.AddMenuEntry(
 			NAME_None,
-			FText::Format(LOCTEXT("DefaultConfiguration",  "Use Project Setting ({0})"), ConfigInfo.Name),
-			ConfigInfo.ToolTip,
+			FText::Format(LOCTEXT("DefaultConfiguration",  "Use Project Setting ({0})"), DefaultConfigInfo.Name),
+			DefaultConfigInfo.ToolTip,
 			FSlateIcon(),
 			FUIAction(
 				FExecuteAction::CreateStatic(&FTurnkeySupportCallbacks::SetPackageBuildConfiguration, VanillaInfo, EProjectPackagingBuildConfigurations::PPBC_MAX),
