Add architecture selection in the editor
- windows only for now

diff --git a/Engine/Source/Developer/DeveloperToolSettings/Classes/Settings/PlatformsMenuSettings.h b/Engine/Source/Developer/DeveloperToolSettings/Classes/Settings/PlatformsMenuSettings.h
--- a/Engine/Source/Developer/DeveloperToolSettings/Classes/Settings/PlatformsMenuSettings.h
+++ b/Engine/Source/Developer/DeveloperToolSettings/Classes/Settings/PlatformsMenuSettings.h
@@ -36,6 +36,9 @@ class DEVELOPERTOOLSETTINGS_API UPlatformsMenuSettings
 	FName GetTargetFlavorForPlatform(FName PlatformName) const;
 	void SetTargetFlavorForPlatform(FName PlatformName, FName TargetFlavorName);
 
+	FString GetArchitectureForPlatform(FName PlatformName) const;
+	void SetArchitectureForPlatform(FName PlatformName, FString ArchitectureName);
+
 	FString GetBuildTargetForPlatform(FName PlatformName) const;
 	void SetBuildTargetForPlatform(FName PlatformName, FString BuildTargetName);
 
@@ -50,6 +53,10 @@ private:
 	UPROPERTY(config)
 	TMap<FName, FName> PerPlatformTargetFlavorName;
 
+	/** Per platform architecture */
+	UPROPERTY(config, EditAnywhere, Category=Project)
+	TMap<FName, FString> PerPlatformArchitecture;
+
 	/** Per platform build target */
 	UPROPERTY(config, EditAnywhere, Category=Project)
 	TMap<FName, FString> PerPlatformBuildTarget;
diff --git a/Engine/Source/Developer/DeveloperToolSettings/Private/PlatformsMenuSettings.cpp b/Engine/Source/Developer/DeveloperToolSettings/Private/PlatformsMenuSettings.cpp
--- a/Engine/Source/Developer/DeveloperToolSettings/Private/PlatformsMenuSettings.cpp
+++ b/Engine/Source/Developer/DeveloperToolSettings/Private/PlatformsMenuSettings.cpp
@@ -59,6 +59,17 @@ void UPlatformsMenuSettings::SetTargetFlavorForPlatform(FName PlatformName, FNam
 	PerPlatformTargetFlavorName.Add(PlatformName, TargetFlavorName);
 }
 
+FString UPlatformsMenuSettings::GetArchitectureForPlatform(FName PlatformName) const
+{
+	const FString* Value = PerPlatformArchitecture.Find(PlatformName);
+	return Value == nullptr ? TEXT("") : *Value;
+}
+
+void UPlatformsMenuSettings::SetArchitectureForPlatform(FName PlatformName, FString ArchitectureName)
+{
+	PerPlatformArchitecture.Add(PlatformName, ArchitectureName);
+}
+
 FString UPlatformsMenuSettings::GetBuildTargetForPlatform(FName PlatformName) const
 {
 	const FString* Value = PerPlatformBuildTarget.Find(PlatformName);
diff --git a/Engine/Source/Developer/TargetPlatform/Private/TargetPlatformSettingsBase.cpp b/Engine/Source/Developer/TargetPlatform/Private/TargetPlatformSettingsBase.cpp
--- a/Engine/Source/Developer/TargetPlatform/Private/TargetPlatformSettingsBase.cpp
+++ b/Engine/Source/Developer/TargetPlatform/Private/TargetPlatformSettingsBase.cpp
@@ -148,4 +148,13 @@ void FTargetPlatformSettingsBase::GetRayTracingShaderFormats(TArray<FName>& OutF
 	{
 		GetAllTargetedShaderFormats(OutFormats);
 	}
-}
\ No newline at end of file
+}
+
+void FTargetPlatformSettingsBase::GetPossibleArchitectures(TArray<FString>& OutArchitectures) const
+{
+}
+
+FString FTargetPlatformSettingsBase::GetHostArchitecture() const
+{
+	return FString();
+}
diff --git a/Engine/Source/Developer/TargetPlatform/Public/Common/TargetPlatformBase.h b/Engine/Source/Developer/TargetPlatform/Public/Common/TargetPlatformBase.h
--- a/Engine/Source/Developer/TargetPlatform/Public/Common/TargetPlatformBase.h
+++ b/Engine/Source/Developer/TargetPlatform/Public/Common/TargetPlatformBase.h
@@ -889,6 +889,14 @@ public:
 	{
 		TargetPlatformSettings->GetRayTracingShaderFormats(OutFormats);
 	}
+	virtual void GetPossibleArchitectures(TArray<FString>& OutArchitectures) const override
+	{
+		TargetPlatformSettings->GetPossibleArchitectures(OutArchitectures);
+	}
+	virtual FString GetHostArchitecture() const override
+	{
+		return TargetPlatformSettings->GetHostArchitecture();
+	}
 	virtual void GetPlatformSpecificProjectAnalytics(TArray<struct FAnalyticsEventAttribute>& AnalyticsParamArray) const override 
 	{
 		TargetPlatformControls->GetPlatformSpecificProjectAnalytics(AnalyticsParamArray);
diff --git a/Engine/Source/Developer/TargetPlatform/Public/Common/TargetPlatformSettingsBase.h b/Engine/Source/Developer/TargetPlatform/Public/Common/TargetPlatformSettingsBase.h
--- a/Engine/Source/Developer/TargetPlatform/Public/Common/TargetPlatformSettingsBase.h
+++ b/Engine/Source/Developer/TargetPlatform/Public/Common/TargetPlatformSettingsBase.h
@@ -41,6 +41,8 @@ public:
 	TARGETPLATFORM_API virtual bool UsesASTCHDR() const override;
 	TARGETPLATFORM_API virtual FName GetFallbackASTCHDR() const override;
 	TARGETPLATFORM_API virtual void GetRayTracingShaderFormats(TArray<FName>& OutFormats) const override;
+	TARGETPLATFORM_API virtual void GetPossibleArchitectures(TArray<FString>& OutArchitectures) const override;
+	TARGETPLATFORM_API virtual FString GetHostArchitecture() const override;
 
 #if WITH_ENGINE
 	virtual void GetShaderFormatModuleHints(TArray<FName>& OutModuleNames) const override
diff --git a/Engine/Source/Developer/TargetPlatform/Public/Interfaces/ITargetPlatformSettings.h b/Engine/Source/Developer/TargetPlatform/Public/Interfaces/ITargetPlatformSettings.h
--- a/Engine/Source/Developer/TargetPlatform/Public/Interfaces/ITargetPlatformSettings.h
+++ b/Engine/Source/Developer/TargetPlatform/Public/Interfaces/ITargetPlatformSettings.h
@@ -275,6 +275,20 @@ public:
 	 */
 	virtual bool SupportsFeature(ETargetPlatformFeatures Feature) const = 0;
 
+	/**
+	 * Get the possible architectures for the target platform
+	 *
+	 * @param OutArchitectures will contain supported architectures.
+	 */
+	virtual void GetPossibleArchitectures(TArray<FString>& OutArchitectures) const = 0;
+
+	/**
+	 * Get the host architectures for the target platform, if any
+	 *
+	 * @return host architecture, or an empty string if it is not relevant
+	 */
+	virtual FString GetHostArchitecture() const = 0;
+
 #if WITH_ENGINE
 	/**
 	 * Gets a list of modules that may contain the GetAllTargetedShaderFormats. This is optional -
diff --git a/Engine/Source/Developer/Windows/WindowsTargetPlatformSettings/Public/GenericWindowsTargetPlatformSettings.h b/Engine/Source/Developer/Windows/WindowsTargetPlatformSettings/Public/GenericWindowsTargetPlatformSettings.h
--- a/Engine/Source/Developer/Windows/WindowsTargetPlatformSettings/Public/GenericWindowsTargetPlatformSettings.h
+++ b/Engine/Source/Developer/Windows/WindowsTargetPlatformSettings/Public/GenericWindowsTargetPlatformSettings.h
@@ -11,6 +11,7 @@
 	#include "DataDrivenShaderPlatformInfo.h"
 #endif // WITH_ENGINE
 #include "HAL/IConsoleManager.h"
+#include "HAL/PlatformMisc.h"
 
 #define LOCTEXT_NAMESPACE "TGenericWindowsTargetPlatformSettings"
 
@@ -193,6 +194,20 @@ public:
 		}
 	}
 
+	virtual void GetPossibleArchitectures(TArray<FString>& OutArchitectures) const override
+	{
+		// @todo: add support for ini-driven 'project supported architectures' to filter this list
+		OutArchitectures.Add(TEXT("x64"));
+		OutArchitectures.Add(TEXT("arm64"));
+		OutArchitectures.Add(TEXT("arm64ec"));
+	}
+
+	virtual FString GetHostArchitecture() const override
+	{
+		return FPlatformMisc::GetHostArchitecture();
+	}
+
+
 #if WITH_ENGINE
 	virtual void GetReflectionCaptureFormats(TArray<FName>& OutFormats) const override
 	{
diff --git a/Engine/Source/Editor/TurnkeySupport/Private/TurnkeySupportModule.cpp b/Engine/Source/Editor/TurnkeySupport/Private/TurnkeySupportModule.cpp
--- a/Engine/Source/Editor/TurnkeySupport/Private/TurnkeySupportModule.cpp
+++ b/Engine/Source/Editor/TurnkeySupport/Private/TurnkeySupportModule.cpp
@@ -398,6 +398,40 @@ public:
 		const ITargetPlatform* TargetPlatform = GetTargetPlatformManager()->FindTargetPlatform(PlatformInfo->Name);
 		TargetPlatform->GetPlatformSpecificProjectAnalytics( AnalyticsParamArray );
 
+
+		// append the architecture
+		FString Architecture = PlatformsSettings->GetArchitectureForPlatform(IniPlatformName);
+		if (!Architecture.IsEmpty())
+		{
+			// special case for host architecture for now : should eventually be handled by UBT
+			if (Architecture == TEXT("Host"))
+			{
+				Architecture = TargetPlatform->GetHostArchitecture();
+			}
+			
+			switch (BuildTargetInfo->Type)
+			{
+				case EBuildTargetType::Server:
+				{
+					BuildCookRunParams += FString::Printf(TEXT(" -serverarchitecture=%s"), *Architecture);
+				}
+				break;
+
+				case EBuildTargetType::Editor:
+				{
+					BuildCookRunParams += FString::Printf(TEXT(" -editorarchitecture=%s"), *Architecture);
+				}
+				break;
+
+				default:
+				{
+					BuildCookRunParams += FString::Printf(TEXT(" -clientarchitecture=%s"), *Architecture);
+				}
+				break;
+			}
+		}
+
+
 		// per mode settings
 		FText ContentPrepDescription;
 		FText ContentPrepTaskName;
@@ -743,6 +777,23 @@ public:
 		return GetDefault<UPlatformsMenuSettings>()->GetTargetFlavorForPlatform(Info->IniPlatformName) == Info->Name;
 	}
 
+	static void SetActiveArchitecture(const PlatformInfo::FTargetPlatformInfo* VanillaInfo, FString Architecture)
+	{
+		UPlatformsMenuSettings* PlatformsSettings = GetMutableDefault<UPlatformsMenuSettings>();
+		PlatformsSettings->SetArchitectureForPlatform(VanillaInfo->IniPlatformName, Architecture);
+		PlatformsSettings->SaveConfig();
+	}
+
+	static bool CanSetActiveArchitecture(const PlatformInfo::FTargetPlatformInfo* Info, FString Architecture)
+	{
+		return true;
+	}
+
+	static bool SetActiveArchitectureIsChecked(const PlatformInfo::FTargetPlatformInfo* Info, FString Architecture)
+	{
+		return GetDefault<UPlatformsMenuSettings>()->GetArchitectureForPlatform(Info->IniPlatformName) == Architecture;
+	}
+
 	static void SetPackageBuildTarget(const PlatformInfo::FTargetPlatformInfo* Info, FString TargetName)
 	{
 		UPlatformsMenuSettings* PlatformsSettings = GetMutableDefault<UPlatformsMenuSettings>();
@@ -1186,6 +1237,8 @@ static void MakeTurnkeyPlatformMenu(UToolMenu* ToolMenu, FName IniPlatformName,
 
 	const PlatformInfo::FTargetPlatformInfo* VanillaInfo = PlatformInfo::FindVanillaPlatformInfo(IniPlatformName);
 
+	const ITargetPlatform* TargetPlatform = GetTargetPlatformManager()->FindTargetPlatform(IniPlatformName);
+
 	if (VanillaInfo != nullptr)
 	{
 		FToolMenuSection& Section = ToolMenu->AddSection("ContentManagement", LOCTEXT("TurnkeySection_Content", "Content Management"));
@@ -1400,6 +1453,67 @@ static void MakeTurnkeyPlatformMenu(UToolMenu* ToolMenu, FName IniPlatformName,
 			}
 		}
 
+
+		{
+			// gather all possible architectures
+			TArray<FString> Architectures;
+			TargetPlatform->GetPossibleArchitectures(Architectures);
+			
+			if (Architectures.Num() > 1)
+			{
+				FToolMenuSection& ArchitectureSection = ToolMenu->AddSection("ArchitectureSelection", LOCTEXT("TurnkeySection_ArchitectureSelection", "Architecture Selection"));
+
+				ArchitectureSection.AddMenuEntry(
+					NAME_None,
+					LOCTEXT("DefaultArchitectureName", "Project Default"),
+					FText(),
+					FSlateIcon(),
+					FUIAction(
+						FExecuteAction::CreateStatic(&FTurnkeySupportCallbacks::SetActiveArchitecture, VanillaInfo, FString()),
+						FCanExecuteAction::CreateStatic(&FTurnkeySupportCallbacks::CanSetActiveArchitecture, VanillaInfo, FString()),
+						FIsActionChecked::CreateStatic(&FTurnkeySupportCallbacks::SetActiveArchitectureIsChecked, VanillaInfo, FString())
+					),
+					EUserInterfaceActionType::RadioButton
+				);
+
+				FString HostArch = TargetPlatform->GetHostArchitecture();
+				if (!HostArch.IsEmpty())
+				{
+					const FString Host(TEXT("Host"));
+
+					ArchitectureSection.AddMenuEntry(
+						NAME_None,
+						FText::Format( LOCTEXT("HostArchitectureName", "Same As Host ({0})"), FText::FromString(HostArch) ),
+						FText(),
+						FSlateIcon(),
+						FUIAction(
+							FExecuteAction::CreateStatic(&FTurnkeySupportCallbacks::SetActiveArchitecture, VanillaInfo, Host),
+							FCanExecuteAction::CreateStatic(&FTurnkeySupportCallbacks::CanSetActiveArchitecture, VanillaInfo, Host),
+							FIsActionChecked::CreateStatic(&FTurnkeySupportCallbacks::SetActiveArchitectureIsChecked, VanillaInfo, Host)
+						),
+						EUserInterfaceActionType::RadioButton
+					);
+				}
+
+				for (const FString& Architecture : Architectures)
+				{
+					ArchitectureSection.AddMenuEntry(
+						NAME_None,
+						FText::FromString(Architecture),
+						FText(),
+						FSlateIcon(),
+						FUIAction(
+							FExecuteAction::CreateStatic(&FTurnkeySupportCallbacks::SetActiveArchitecture, VanillaInfo, Architecture),
+							FCanExecuteAction::CreateStatic(&FTurnkeySupportCallbacks::CanSetActiveArchitecture, VanillaInfo, Architecture),
+							FIsActionChecked::CreateStatic(&FTurnkeySupportCallbacks::SetActiveArchitectureIsChecked, VanillaInfo, Architecture)
+						),
+						EUserInterfaceActionType::RadioButton
+					);
+				}
+			}
+		}
+
+
 		FToolMenuSection& DevicesSection = ToolMenu->AddSection("AllDevices", LOCTEXT("TurnkeySection_AllDevices", "All Devices"));
 
 		TArray<TSharedPtr<ITargetDeviceProxy>> DeviceProxies;
diff --git a/Engine/Source/Runtime/Core/Private/GenericPlatform/GenericPlatformMisc.cpp b/Engine/Source/Runtime/Core/Private/GenericPlatform/GenericPlatformMisc.cpp
--- a/Engine/Source/Runtime/Core/Private/GenericPlatform/GenericPlatformMisc.cpp
+++ b/Engine/Source/Runtime/Core/Private/GenericPlatform/GenericPlatformMisc.cpp
@@ -676,6 +676,11 @@ bool FGenericPlatformMisc::NeedsNonoptionalCPUFeaturesCheck()
 	return false;
 }
 
+FString FGenericPlatformMisc::GetHostArchitecture()
+{
+	return FString();
+}
+
 FString FGenericPlatformMisc::GetPrimaryGPUBrand()
 {
 	// Not implemented cross-platform. Each platform may or may not choose to implement this.
diff --git a/Engine/Source/Runtime/Core/Private/Windows/WindowsPlatformMisc.cpp b/Engine/Source/Runtime/Core/Private/Windows/WindowsPlatformMisc.cpp
--- a/Engine/Source/Runtime/Core/Private/Windows/WindowsPlatformMisc.cpp
+++ b/Engine/Source/Runtime/Core/Private/Windows/WindowsPlatformMisc.cpp
@@ -4236,6 +4236,38 @@ FString FWindowsPlatformMisc::GetOperatingSystemId()
 	return Result;
 }
 
+FString FWindowsPlatformMisc::GetHostArchitecture()
+{
+#if defined(_M_ARM64EC) || defined(_M_ARM64)
+
+	// these build targets can only run on arm64 (assuming there's no emulation)
+	return TEXT("arm64");
+
+#else
+
+	// try to find new Windows 10+ API
+#pragma warning( push )
+#pragma warning( disable: 4191 )	// unsafe conversion from 'type of expression' to 'type required'
+#pragma warning( disable: 6387 )	// 'GetModuleHandleW("kernel32")' could be '0'
+	typedef BOOL( WINAPI *LPFN_IsWow64Process2)(HANDLE, USHORT*, USHORT*);
+	LPFN_IsWow64Process2 fnIsWow64Process2 = (LPFN_IsWow64Process2)GetProcAddress( GetModuleHandle(TEXT("kernel32")), "IsWow64Process2" );
+#pragma warning( pop )
+
+	USHORT ProcessMachine, NativeMachine;
+	if (fnIsWow64Process2 != nullptr && fnIsWow64Process2(::GetCurrentProcess(), &ProcessMachine, &NativeMachine) )
+	{
+		if (NativeMachine == IMAGE_FILE_MACHINE_ARM64)
+		{
+			return TEXT("arm64");
+		}
+	}
+
+	// assume normal x64 by default
+	return "x64";
+
+#endif
+}
+
 
 EConvertibleLaptopMode FWindowsPlatformMisc::GetConvertibleLaptopMode()
 {
diff --git a/Engine/Source/Runtime/Core/Public/GenericPlatform/GenericPlatformMisc.h b/Engine/Source/Runtime/Core/Public/GenericPlatform/GenericPlatformMisc.h
--- a/Engine/Source/Runtime/Core/Public/GenericPlatform/GenericPlatformMisc.h
+++ b/Engine/Source/Runtime/Core/Public/GenericPlatform/GenericPlatformMisc.h
@@ -737,6 +737,9 @@ struct FGenericPlatformMisc
 	/** @return whether to check for specific CPU compatibility or not */
 	static CORE_API bool NeedsNonoptionalCPUFeaturesCheck();
 
+	/** @return architecture of the host machine (x64, arm64 etc) **/
+	static CORE_API FString GetHostArchitecture();
+
 	/**
 	 * Uses cpuid instruction to get the CPU brand string
 	 *
diff --git a/Engine/Source/Runtime/Core/Public/Windows/WindowsPlatformMisc.h b/Engine/Source/Runtime/Core/Public/Windows/WindowsPlatformMisc.h
--- a/Engine/Source/Runtime/Core/Public/Windows/WindowsPlatformMisc.h
+++ b/Engine/Source/Runtime/Core/Public/Windows/WindowsPlatformMisc.h
@@ -284,6 +284,9 @@ struct FWindowsPlatformMisc
 	/** @return whether this cpu has timed pause instruction support or not */
 	static CORE_API bool HasTimedPauseCPUFeature();
 
+	/** @return architecture of the host machine (x64, arm64) **/
+	static CORE_API FString GetHostArchitecture();
+
 	/** 
 	 * Provides a simpler interface for fetching and cleanup of registry value queries
 	 *
