Simplify Host Architecture display in the platforms menu
- just note it next to the architecture name, rather than adding a separate menu item

#rnx

diff --git a/Engine/Source/Editor/TurnkeySupport/Private/TurnkeySupportModule.cpp b/Engine/Source/Editor/TurnkeySupport/Private/TurnkeySupportModule.cpp
--- a/Engine/Source/Editor/TurnkeySupport/Private/TurnkeySupportModule.cpp
+++ b/Engine/Source/Editor/TurnkeySupport/Private/TurnkeySupportModule.cpp
@@ -403,12 +403,6 @@ public:
 		FString Architecture = PlatformsSettings->GetArchitectureForPlatform(IniPlatformName);
 		if (!Architecture.IsEmpty())
 		{
-			// special case for host architecture for now : should eventually be handled by UBT
-			if (Architecture == TEXT("Host"))
-			{
-				Architecture = TargetPlatform->GetHostArchitecture();
-			}
-			
 			switch (BuildTargetInfo->Type)
 			{
 				case EBuildTargetType::Server:
@@ -1476,30 +1470,17 @@ static void MakeTurnkeyPlatformMenu(UToolMenu* ToolMenu, FName IniPlatformName,
 					EUserInterfaceActionType::RadioButton
 				);
 
-				FString HostArch = TargetPlatform->GetHostArchitecture();
-				if (!HostArch.IsEmpty())
+				for (const FString& Architecture : Architectures)
 				{
-					const FString Host(TEXT("Host"));
+					bool bIsHostArchitecture = (TargetPlatform->IniPlatformName() == FPlatformProperties::IniPlatformName()) && (Architecture == TargetPlatform->GetHostArchitecture());
 
-					ArchitectureSection.AddMenuEntry(
-						NAME_None,
-						FText::Format( LOCTEXT("HostArchitectureName", "Same As Host ({0})"), FText::FromString(HostArch) ),
-						FText(),
-						FSlateIcon(),
-						FUIAction(
-							FExecuteAction::CreateStatic(&FTurnkeySupportCallbacks::SetActiveArchitecture, VanillaInfo, Host),
-							FCanExecuteAction::CreateStatic(&FTurnkeySupportCallbacks::CanSetActiveArchitecture, VanillaInfo, Host),
-							FIsActionChecked::CreateStatic(&FTurnkeySupportCallbacks::SetActiveArchitectureIsChecked, VanillaInfo, Host)
-						),
-						EUserInterfaceActionType::RadioButton
-					);
-				}
+					FText ArchLabel = bIsHostArchitecture 
+						? FText::Format(LOCTEXT("HostArchLabel", "{0} (this platform)"), FText::FromString(Architecture)) 
+						: FText::FromString(Architecture);
 
-				for (const FString& Architecture : Architectures)
-				{
 					ArchitectureSection.AddMenuEntry(
 						NAME_None,
-						FText::FromString(Architecture),
+						ArchLabel,
 						FText(),
 						FSlateIcon(),
 						FUIAction(
