Launch from the Platform menu will now build an x86_64 binary when launching on x86_64 devices such as the emulator

#android

diff --git a/Engine/Source/Developer/Android/AndroidDeviceDetection/Private/AndroidDeviceDetectionModule.cpp b/Engine/Source/Developer/Android/AndroidDeviceDetection/Private/AndroidDeviceDetectionModule.cpp
--- a/Engine/Source/Developer/Android/AndroidDeviceDetection/Private/AndroidDeviceDetectionModule.cpp
+++ b/Engine/Source/Developer/Android/AndroidDeviceDetection/Private/AndroidDeviceDetectionModule.cpp
@@ -261,6 +261,30 @@ private:
 			return;
 		}
 
+		auto AbiToArchitecture = [](FStringView Abi) -> FStringView
+		{
+			if (Abi == TEXTVIEW("arm64-v8a"))
+			{
+				return TEXTVIEW("arm64");
+			}
+			else if (Abi == TEXTVIEW("armeabi-v7a"))
+			{
+				return TEXTVIEW("arm32");
+			}
+			else if (Abi == TEXTVIEW("x86_64"))
+			{
+				return TEXTVIEW("x64");
+			}
+			else if (Abi == TEXTVIEW("x86"))
+			{
+				return TEXTVIEW("x86");
+			}
+			else
+			{
+				return {};
+			}
+		};
+
 		// separate out each line
 		TArray<FString> DeviceStrings;
 		StdOut = StdOut.Replace(TEXT("\r"), TEXT("\n"));
@@ -384,6 +408,15 @@ private:
 					NewDeviceInfo.SDKVersion = INDEX_NONE;
 				}
 
+				if (FString AbiOutput; !ExecuteAdbCommand(*FString::Printf(TEXT("-s %s shell getprop ro.product.cpu.abi"), *SerialNumber), &AbiOutput, nullptr))
+				{
+					continue;
+				}
+				else
+				{
+					NewDeviceInfo.Architecture = AbiToArchitecture(FStringView{AbiOutput}.TrimStartAndEnd());
+				}
+
 				if (bGetExtensionsViaSurfaceFlinger)
 				{
 					// get the GL extensions string (and a bunch of other stuff)
@@ -723,10 +756,27 @@ private:
 			
 			if (!Device || !Device->SerialNumber.IsEmpty())
 			{
+				FString Architecture;
+				
+				if (!FFileHelper::LoadFileToStringWithLineVisitor(*FPaths::Combine(AvdHomePath, AvdName + TEXT(".avd"), TEXT("config.ini")), [&](FStringView Line)
+				{
+					if (int32 Separator; Line.FindChar('=', Separator))
+					{
+						if (FStringView Key = Line.Left(Separator).TrimStartAndEnd(); Key == TEXT("hw.cpu.arch"))
+						{
+							Architecture = AbiToArchitecture(Line.RightChop(Separator + 1).TrimStartAndEnd());
+						}
+					}
+				}))
+				{
+					continue;
+				}
+
 				FAndroidDeviceInfo NewDeviceInfo
 				{
 					.DeviceId = MoveTemp(DeviceId),
 					.AvdName = MoveTemp(AvdName),
+					.Architecture = MoveTemp(Architecture),
 					.GLESExtensions = TEXT("GL_KHR_texture_compression_astc_ldr"),
 					.GLESVersion = 0x30001,
 					.bAuthorizedDevice = true,
diff --git a/Engine/Source/Developer/Android/AndroidDeviceDetection/Public/Interfaces/IAndroidDeviceDetection.h b/Engine/Source/Developer/Android/AndroidDeviceDetection/Public/Interfaces/IAndroidDeviceDetection.h
--- a/Engine/Source/Developer/Android/AndroidDeviceDetection/Public/Interfaces/IAndroidDeviceDetection.h
+++ b/Engine/Source/Developer/Android/AndroidDeviceDetection/Public/Interfaces/IAndroidDeviceDetection.h
@@ -32,6 +32,8 @@ struct FAndroidDeviceInfo
 	// Android SDK version supported by the device (ro.build.version.sdk - note: deprecated in 4 according to docs, but version 4 devices return an empty string when querying the 'replacement' SDK_INT)
 	int32 SDKVersion = INDEX_NONE;
 
+	FString Architecture;
+
 	// List of supported OpenGL extensions (retrieved via SurfaceFlinger)
 	FString GLESExtensions;
 
diff --git a/Engine/Source/Developer/Android/AndroidTargetPlatformControls/Private/AndroidTargetDevice.h b/Engine/Source/Developer/Android/AndroidTargetPlatformControls/Private/AndroidTargetDevice.h
--- a/Engine/Source/Developer/Android/AndroidTargetPlatformControls/Private/AndroidTargetDevice.h
+++ b/Engine/Source/Developer/Android/AndroidTargetPlatformControls/Private/AndroidTargetDevice.h
@@ -96,6 +96,16 @@ public:
 		AndroidVersionString = InReleaseVersion;
 	}
 
+	/**
+	 * Sets the architecture.
+	 *
+	 * @param InArchitecture - The architecture to set.
+	 */
+	void SetArchitecture(FString InArchitecture)
+	{
+		Architecture = MoveTemp(InArchitecture);
+	}
+
 	/**
 	 * Sets the name.
 	 *
@@ -180,6 +190,11 @@ public:
 		return AndroidVersionString;
 	}
 
+	virtual FStringView GetArchitecture() const override
+	{
+		return Architecture;
+	}
+
 	virtual const class ITargetPlatformSettings& GetPlatformSettings() const override
 	{
 		return *(TargetPlatformControls.GetTargetPlatformSettings());
@@ -273,6 +288,8 @@ protected:
 	// Holds the Android Release version string (e.g., "2.3" or "4.2.2")
 	FString AndroidVersionString;
 
+	FString Architecture;
+	
 	// Holds the name.
 	FString Name;
 	
diff --git a/Engine/Source/Developer/Android/AndroidTargetPlatformControls/Private/AndroidTargetPlatformControls.cpp b/Engine/Source/Developer/Android/AndroidTargetPlatformControls/Private/AndroidTargetPlatformControls.cpp
--- a/Engine/Source/Developer/Android/AndroidTargetPlatformControls/Private/AndroidTargetPlatformControls.cpp
+++ b/Engine/Source/Developer/Android/AndroidTargetPlatformControls/Private/AndroidTargetPlatformControls.cpp
@@ -446,6 +446,7 @@ bool FAndroidTargetPlatformControls::HandleTicker( float DeltaTime )
 			Device->SetDeviceName(DeviceInfo.DeviceName);
 			Device->SetAuthorized(DeviceInfo.bAuthorizedDevice);
 			Device->SetVersions(DeviceInfo.SDKVersion, DeviceInfo.HumanAndroidVersion);
+			Device->SetArchitecture(DeviceInfo.Architecture);
 			Device->SetSerialNumber(DeviceInfo.SerialNumber);
 
 			ITargetPlatformControls::OnDeviceDiscovered().Broadcast(Device.ToSharedRef());
diff --git a/Engine/Source/Developer/TargetDeviceServices/Private/Proxies/TargetDeviceProxy.cpp b/Engine/Source/Developer/TargetDeviceServices/Private/Proxies/TargetDeviceProxy.cpp
--- a/Engine/Source/Developer/TargetDeviceServices/Private/Proxies/TargetDeviceProxy.cpp
+++ b/Engine/Source/Developer/TargetDeviceServices/Private/Proxies/TargetDeviceProxy.cpp
@@ -65,6 +65,7 @@ void FTargetDeviceProxy::UpdateFromMessage( const FTargetDeviceServicePong& Mess
 	DeviceUser = Message.DeviceUser;
 	DeviceUserPassword = Message.DeviceUserPassword;
 	OSVersion = Message.OSVersion;
+	Architecture = Message.Architecture;
 	ConnectionType = Message.ConnectionType;
 	Shared = Message.Shared;
 	SupportsMultiLaunch = Message.SupportsMultiLaunch;
diff --git a/Engine/Source/Developer/TargetDeviceServices/Private/Proxies/TargetDeviceProxy.h b/Engine/Source/Developer/TargetDeviceServices/Private/Proxies/TargetDeviceProxy.h
--- a/Engine/Source/Developer/TargetDeviceServices/Private/Proxies/TargetDeviceProxy.h
+++ b/Engine/Source/Developer/TargetDeviceServices/Private/Proxies/TargetDeviceProxy.h
@@ -161,6 +161,11 @@ public:
 		return OSVersion;
 	}
 
+	virtual FStringView GetArchitecture() const override
+	{
+		return Architecture;
+	}
+	
 	virtual const FString& GetConnectionType() const override
 	{
 		return ConnectionType;
@@ -266,6 +271,9 @@ private:
 	/** Holds the device type. */
 	FString Type;
 
+	/** Holds the remote device's architecture. */
+	FString Architecture;
+	
 	/** Holds the remote device's OS version. */
 	FString OSVersion;
 
diff --git a/Engine/Source/Developer/TargetDeviceServices/Private/Services/TargetDeviceService.cpp b/Engine/Source/Developer/TargetDeviceServices/Private/Services/TargetDeviceService.cpp
--- a/Engine/Source/Developer/TargetDeviceServices/Private/Services/TargetDeviceService.cpp
+++ b/Engine/Source/Developer/TargetDeviceServices/Private/Services/TargetDeviceService.cpp
@@ -431,6 +431,7 @@ void FTargetDeviceService::HandlePingMessage(const FTargetDeviceServicePing& InM
 		Message->Make = TEXT("@todo");
 		Message->Model = DefaultDevice->GetModelId();
 		Message->OSVersion = DefaultDevice->GetOSVersion();
+		Message->Architecture = DefaultDevice->GetArchitecture();
 		DefaultDevice->GetUserCredentials(Message->DeviceUser, Message->DeviceUserPassword);
 		Message->Shared = Shared;
 		Message->SupportsMultiLaunch = DefaultDevice->SupportsFeature(ETargetDeviceFeatures::MultiLaunch);
diff --git a/Engine/Source/Developer/TargetDeviceServices/Public/ITargetDeviceProxy.h b/Engine/Source/Developer/TargetDeviceServices/Public/ITargetDeviceProxy.h
--- a/Engine/Source/Developer/TargetDeviceServices/Public/ITargetDeviceProxy.h
+++ b/Engine/Source/Developer/TargetDeviceServices/Public/ITargetDeviceProxy.h
@@ -209,10 +209,20 @@ public:
 	/**
 	 * Gets the OS Version of the device
 	 *
-	 * @return Device OS VErsion.
+	 * @return Device OS Version.
 	 */
 	virtual const FString& GetOSVersion() const = 0;
 
+	/**
+	 * Gets the architecture of the device
+	 *
+	 * @return Device Architecture.
+	 */
+	virtual FStringView GetArchitecture() const
+	{
+		return {};
+	}
+
 	/**
 	 * Gets the type of connection for this device.
 	 *
diff --git a/Engine/Source/Developer/TargetDeviceServices/Public/TargetDeviceServiceMessages.h b/Engine/Source/Developer/TargetDeviceServices/Public/TargetDeviceServiceMessages.h
--- a/Engine/Source/Developer/TargetDeviceServices/Public/TargetDeviceServiceMessages.h
+++ b/Engine/Source/Developer/TargetDeviceServices/Public/TargetDeviceServiceMessages.h
@@ -292,6 +292,10 @@ struct FTargetDeviceServicePong
 	UPROPERTY(EditAnywhere, Category="Message")
 	FString OSVersion;
 
+	/** Holds the device architecture. */
+	UPROPERTY(EditAnywhere, Category="Message")
+	FString Architecture;
+
 	/** Holds the connection type. */
 	UPROPERTY(EditAnywhere, Category="Message")
 	FString ConnectionType;
diff --git a/Engine/Source/Developer/TargetPlatform/Public/Interfaces/ITargetDevice.h b/Engine/Source/Developer/TargetPlatform/Public/Interfaces/ITargetDevice.h
--- a/Engine/Source/Developer/TargetPlatform/Public/Interfaces/ITargetDevice.h
+++ b/Engine/Source/Developer/TargetPlatform/Public/Interfaces/ITargetDevice.h
@@ -303,6 +303,16 @@ public:
 	 */
 	virtual FString GetOSVersion() const {return FString();};
 
+	/**
+	 * Gets the device architecture.
+	 *
+	 * @return Device Architecture.
+	 */
+	virtual FStringView GetArchitecture() const
+	{
+		return {};
+	}
+	
 	/**
 	 * Gets the device connection type.
 	 * 
diff --git a/Engine/Source/Editor/TurnkeySupport/Private/TurnkeyEditorSupport.cpp b/Engine/Source/Editor/TurnkeySupport/Private/TurnkeyEditorSupport.cpp
--- a/Engine/Source/Editor/TurnkeySupport/Private/TurnkeyEditorSupport.cpp
+++ b/Engine/Source/Editor/TurnkeySupport/Private/TurnkeyEditorSupport.cpp
@@ -92,7 +92,7 @@ void FTurnkeyEditorSupport::PrepareToLaunchRunningMap(const FString& DeviceId, c
 #endif
 }
 
-void FTurnkeyEditorSupport::LaunchRunningMap(const FString& DeviceId, const FString& DeviceName, const FString& ProjectPath, bool bUseTurnkey, bool bOnSimulator)
+void FTurnkeyEditorSupport::LaunchRunningMap(const FString& DeviceId, const FString& DeviceName, const FString& ProjectPath, bool bUseTurnkey, bool bOnSimulator, FStringView Architecture)
 {
 #if WITH_EDITOR
 	FTargetDeviceId TargetDeviceId;
@@ -133,10 +133,11 @@ void FTurnkeyEditorSupport::LaunchRunningMap(const FString& DeviceId, const FStr
 			// to get an error back from CreateUatTask, then we should set this to bUseTurnkey and remove the block below, and let the code in FLauncherWorker::CreateAndExecuteTasks handle it
 			DeviceInfo.bUpdateDeviceFlash = false;
 			DeviceInfo.bIsSimulator = bOnSimulator;
+			DeviceInfo.Architecture = Architecture;
 
 			FRequestPlaySessionParams SessionParams;
 			SessionParams.SessionDestination = EPlaySessionDestinationType::Launcher;
-			SessionParams.LauncherTargetDevice = DeviceInfo;
+			SessionParams.LauncherTargetDevice = MoveTemp(DeviceInfo);
 
 			// if we want to check device flash before we start cooking, kick it off now. we could delay this 
 			if (bUseTurnkey)
diff --git a/Engine/Source/Editor/TurnkeySupport/Private/TurnkeyEditorSupport.h b/Engine/Source/Editor/TurnkeySupport/Private/TurnkeyEditorSupport.h
--- a/Engine/Source/Editor/TurnkeySupport/Private/TurnkeyEditorSupport.h
+++ b/Engine/Source/Editor/TurnkeySupport/Private/TurnkeyEditorSupport.h
@@ -14,7 +14,7 @@ public:
 	static FString GetUATOptions();
 
 	static void PrepareToLaunchRunningMap(const FString& DeviceId, const FString& DeviceName);
-	static void LaunchRunningMap(const FString& DeviceId, const FString& DeviceName, const FString& ProjectPath, bool bUseTurnkey, bool bOnSimulator = false);
+	static void LaunchRunningMap(const FString& DeviceId, const FString& DeviceName, const FString& ProjectPath, bool bUseTurnkey, bool bOnSimulator = false, FStringView Architecture = {});
 	static void AddEditorOptions(struct FToolMenuSection& MenuBuilder);
 	
 	static void SaveAll();
diff --git a/Engine/Source/Editor/TurnkeySupport/Private/TurnkeySupportModule.cpp b/Engine/Source/Editor/TurnkeySupport/Private/TurnkeySupportModule.cpp
--- a/Engine/Source/Editor/TurnkeySupport/Private/TurnkeySupportModule.cpp
+++ b/Engine/Source/Editor/TurnkeySupport/Private/TurnkeySupportModule.cpp
@@ -1737,30 +1737,30 @@ bool CanLaunchOnDevice(const FString& DeviceName)
 	return false;
 }
 
-static void LaunchOnDevice(const FString& DeviceId, const FString& DeviceName, bool bUseTurnkey)
+static void LaunchOnDevice(const FString& DeviceId, const FString& DeviceName, bool bUseTurnkey, FStringView Architecture)
 {
-	FTurnkeyEditorSupport::LaunchRunningMap(DeviceId, DeviceName, GetProjectPathForTurnkey(), bUseTurnkey);
+	FTurnkeyEditorSupport::LaunchRunningMap(DeviceId, DeviceName, GetProjectPathForTurnkey(), bUseTurnkey, false, Architecture);
 }
 
-static void LaunchOnSimulator(const FString& DeviceId, const FString& DeviceName, bool bUseTurnkey)
+static void LaunchOnSimulator(const FString& DeviceId, const FString& DeviceName, bool bUseTurnkey, FStringView Architecture)
 {								
-	FTurnkeyEditorSupport::LaunchRunningMap(DeviceId, DeviceName, GetProjectPathForTurnkey(), false, true);
+	FTurnkeyEditorSupport::LaunchRunningMap(DeviceId, DeviceName, GetProjectPathForTurnkey(), false, true, Architecture);
 }
 
-static void PrepareLaunchOn(FString DeviceId, FString DeviceName)
+static void PrepareLaunchOn(const FString& DeviceId, const FString& DeviceName)
 {
 	FTurnkeyEditorSupport::PrepareToLaunchRunningMap(DeviceId, DeviceName);
 }
 
-static void HandleLaunchOnDeviceActionExecute(FString DeviceId, FString DeviceName, bool bUseTurnkey)
+static void HandleLaunchOnDeviceActionExecute(const FString& DeviceId, const FString& DeviceName, bool bUseTurnkey, FStringView Architecture)
 {
 	PrepareLaunchOn(DeviceId, DeviceName);
-	LaunchOnDevice(DeviceId, DeviceName, bUseTurnkey);
+	LaunchOnDevice(DeviceId, DeviceName, bUseTurnkey, Architecture);
 }
 
-static void HandleLaunchOnSimulatorActionExecute(FString DeviceId, FString DeviceName, bool bUseTurnkey)
+static void HandleLaunchOnSimulatorActionExecute(const FString& DeviceId, const FString& DeviceName, bool bUseTurnkey, FStringView Architecture)
 {
-	LaunchOnSimulator(DeviceId, DeviceName, bUseTurnkey);
+	LaunchOnSimulator(DeviceId, DeviceName, bUseTurnkey, Architecture);
 }
 
 static bool HandleLaunchOnDeviceActionCanExecute(FString DeviceName)
@@ -1844,12 +1844,12 @@ static void GenerateDeviceProxyMenuParams(TSharedPtr<ITargetDeviceProxy> DeviceP
 				if (DeviceProxy->IsSimulated())
 				{
 					PerPlatformLastSimChosen[*DeviceProxy->GetType()] = DeviceProxy->GetName();
-					HandleLaunchOnSimulatorActionExecute(DeviceId, DeviceProxy->GetName(), true);
+					HandleLaunchOnSimulatorActionExecute(DeviceId, DeviceProxy->GetName(), true, DeviceProxy->GetArchitecture());
 				}
 				else
 				{
 					PerPlatformLastChosen[PlatformName] = DeviceProxy->GetName();
-					HandleLaunchOnDeviceActionExecute(DeviceId, DeviceProxy->GetName(), true);
+					HandleLaunchOnDeviceActionExecute(DeviceId, DeviceProxy->GetName(), true, DeviceProxy->GetArchitecture());
 				}
 				ExternalOnClickDelegate.ExecuteIfBound(DeviceId);
 			}
@@ -2955,11 +2955,11 @@ void FTurnkeySupportModule::RepeatQuickLaunch(FString DeviceId)
 	{				
 		if (Proxy->IsSimulated())
 		{
-			HandleLaunchOnSimulatorActionExecute(DeviceId, Proxy->GetName(), true);
+			HandleLaunchOnSimulatorActionExecute(DeviceId, Proxy->GetName(), true, Proxy->GetArchitecture());
 		}
 		else
 		{
-			HandleLaunchOnDeviceActionExecute(DeviceId, Proxy->GetName(), true);
+			HandleLaunchOnDeviceActionExecute(DeviceId, Proxy->GetName(), true, Proxy->GetArchitecture());
 		}
 	}
 	else
diff --git a/Engine/Source/Editor/UnrealEd/Private/PlayLevelViaLauncher.cpp b/Engine/Source/Editor/UnrealEd/Private/PlayLevelViaLauncher.cpp
--- a/Engine/Source/Editor/UnrealEd/Private/PlayLevelViaLauncher.cpp
+++ b/Engine/Source/Editor/UnrealEd/Private/PlayLevelViaLauncher.cpp
@@ -302,6 +302,11 @@ void UEditorEngine::StartPlayUsingLauncherSession(FRequestPlaySessionParams& InR
 	LauncherProfile->SetEditorExe(FUnrealEdMisc::Get().GetExecutableForCommandlets());
 	LauncherProfile->SetShouldUpdateDeviceFlash(InRequestParams.LauncherTargetDevice->bUpdateDeviceFlash);
 	LauncherProfile->SetCookOptions(*CookOptions);
+
+	if (!InRequestParams.LauncherTargetDevice->Architecture.IsEmpty())
+	{
+		LauncherProfile->SetClientArchitectures({&InRequestParams.LauncherTargetDevice->Architecture, 1});
+	}
 	
 	if (LauncherProfile->IsBuildingUAT() && !GetDefault<UEditorPerProjectUserSettings>()->bAlwaysBuildUAT && bUATSuccessfullyCompiledOnce)
 	{
diff --git a/Engine/Source/Editor/UnrealEd/Public/PlayInEditorDataTypes.h b/Engine/Source/Editor/UnrealEd/Public/PlayInEditorDataTypes.h
--- a/Engine/Source/Editor/UnrealEd/Public/PlayInEditorDataTypes.h
+++ b/Engine/Source/Editor/UnrealEd/Public/PlayInEditorDataTypes.h
@@ -145,6 +145,8 @@ public:
 		bool bUpdateDeviceFlash = false;
 		/** If True, the launch device is a Simulator */
 		bool bIsSimulator = false;
+		/** If set, the architecture used for building and deploying to the device */
+		FString Architecture;
 	};
 
 	/** Where should the session be launched? May be local or remote. */
diff --git a/Engine/Source/Programs/AutomationTool/Android/AndroidPlatform.Automation.cs b/Engine/Source/Programs/AutomationTool/Android/AndroidPlatform.Automation.cs
--- a/Engine/Source/Programs/AutomationTool/Android/AndroidPlatform.Automation.cs
+++ b/Engine/Source/Programs/AutomationTool/Android/AndroidPlatform.Automation.cs
@@ -4146,7 +4146,7 @@ public class AndroidPlatform : Platform
 		}
 
 		// @todo get TargetName?
-		UnrealArchitectures AppArchitectures = UnrealArchitectureConfig.ForPlatform(UnrealTargetPlatform.Android).ActiveArchitectures(Params.RawProjectPath, null);
+		UnrealArchitectures AppArchitectures = Params.ClientArchitecture ?? UnrealArchitectureConfig.ForPlatform(UnrealTargetPlatform.Android).ActiveArchitectures(Params.RawProjectPath, null);
 
 		// ask the device
 		IProcessResult ABIResult = RunAdbCommand(Params, DeviceName, " shell getprop ro.product.cpu.abi", null, ERunOptions.AppMustExist);
