crash fix for null TargetPlatform
- typically when the platform's Target Platform Module does not exist but the DDPI does.

#rnx

diff --git a/Engine/Source/Editor/TurnkeySupport/Private/TurnkeySupportModule.cpp b/Engine/Source/Editor/TurnkeySupport/Private/TurnkeySupportModule.cpp
--- a/Engine/Source/Editor/TurnkeySupport/Private/TurnkeySupportModule.cpp
+++ b/Engine/Source/Editor/TurnkeySupport/Private/TurnkeySupportModule.cpp
@@ -1446,7 +1446,7 @@ static void MakeTurnkeyPlatformMenu(UToolMenu* ToolMenu, FName IniPlatformName,
 			}
 		}
 
-
+		if (TargetPlatform != nullptr)
 		{
 			// gather all possible architectures
 			TArray<FString> Architectures;
