[Editor] Moves Platforms from toolbar to the main menu

Nothing within the menu has changed, just the location. Since the platforms menu was hidden during PIE, technically this means you can use the platforms menu now while in PIE.

I removed the separator on the toolbar between the PIE controlls and the debugger buttons (step into, step out, etc.). Previously the debugger buttons and the platforms menu both came after the separator, but the debugger buttons were hidden when not stopped at a breakpoint. With the platform menu gone, this left an empty space between two separators by default when in the editor.

diff --git a/Engine/Source/Editor/LevelEditor/LevelEditor.Build.cs b/Engine/Source/Editor/LevelEditor/LevelEditor.Build.cs
--- a/Engine/Source/Editor/LevelEditor/LevelEditor.Build.cs
+++ b/Engine/Source/Editor/LevelEditor/LevelEditor.Build.cs
@@ -22,6 +22,7 @@ public class LevelEditor : ModuleRules
 				"WorldBrowser",
 				"NewLevelDialog",
 				"LocalizationDashboard",
+				"TurnkeySupport",
 			}
 		);
 
diff --git a/Engine/Source/Editor/LevelEditor/Private/LevelEditorMenu.cpp b/Engine/Source/Editor/LevelEditor/Private/LevelEditorMenu.cpp
--- a/Engine/Source/Editor/LevelEditor/Private/LevelEditorMenu.cpp
+++ b/Engine/Source/Editor/LevelEditor/Private/LevelEditorMenu.cpp
@@ -34,6 +34,7 @@
 #include "LevelViewportActions.h"
 #include "Styling/SlateIconFinder.h"
 #include "Toolkits/GlobalEditorCommonCommands.h"
+#include "Interfaces/ITurnkeySupportModule.h"
 
 #define LOCTEXT_NAMESPACE "LevelEditorMenu"
 
@@ -354,6 +355,10 @@ void FLevelEditorMenu::RegisterLevelEditorMenus()
 					}));
 			
 			ActionsEntry.InsertPosition = FToolMenuInsert("Help", EToolMenuInsertType::Before);
+
+#if UE_WITH_TURNKEY_SUPPORT
+			ITurnkeySupportModule::Get().MakeTurnkeyMenu(Section, FToolMenuInsert("Select", EToolMenuInsertType::Before));
+#endif
 		}
 
 	};
diff --git a/Engine/Source/Editor/LevelEditor/Private/LevelEditorToolBar.cpp b/Engine/Source/Editor/LevelEditor/Private/LevelEditorToolBar.cpp
--- a/Engine/Source/Editor/LevelEditor/Private/LevelEditorToolBar.cpp
+++ b/Engine/Source/Editor/LevelEditor/Private/LevelEditorToolBar.cpp
@@ -1517,7 +1517,7 @@ void FLevelEditorToolBar::RegisterLevelEditorToolBar( const TSharedRef<FUIComman
 		PreviewModeFunctionality::AddPreviewToggleButton(PlaySection);
 
 		// Add the shared play-world commands that will be shown on the Kismet toolbar as well
-		FPlayWorldCommands::BuildToolbar(PlaySection, true);
+		FPlayWorldCommands::BuildToolbar(PlaySection);
 
 
 	}
diff --git a/Engine/Source/Editor/TurnkeySupport/Private/TurnkeySupportModule.cpp b/Engine/Source/Editor/TurnkeySupport/Private/TurnkeySupportModule.cpp
--- a/Engine/Source/Editor/TurnkeySupport/Private/TurnkeySupportModule.cpp
+++ b/Engine/Source/Editor/TurnkeySupport/Private/TurnkeySupportModule.cpp
@@ -2240,9 +2240,7 @@ TSharedRef<SWidget> FTurnkeySupportModule::MakeTurnkeyMenuWidget() const
 	FTurnkeySupportCommands::Register();
 	const FTurnkeySupportCommands& Commands = FTurnkeySupportCommands::Get();
 
-	const bool bShouldCloseWindowAfterMenuSelection = true;
-
-	static const FName MenuName("UnrealEd.PlayWorldCommands.PlatformsMenu");
+	static const FName MenuName("LevelEditor.MainMenu.Platforms");
 
 	if (!UToolMenus::Get()->IsMenuRegistered(MenuName))
 	{
@@ -2415,35 +2413,20 @@ TSharedRef<SWidget> FTurnkeySupportModule::MakeTurnkeyMenuWidget() const
 
 }
 
-void FTurnkeySupportModule::MakeTurnkeyMenu(FToolMenuSection& MenuSection) const
+void FTurnkeySupportModule::MakeTurnkeyMenu(FToolMenuSection& MenuSection, const FToolMenuInsert& InsertPosition) const
 {
 	// make sure the DeviceProxyManager is going _before_ we create the menu contents dynamically, so that devices will show up
 	ITargetDeviceServicesModule* TargetDeviceServicesModule = static_cast<ITargetDeviceServicesModule*>(FModuleManager::Get().LoadModule(TEXT("TargetDeviceServices")));
 	TargetDeviceServicesModule->GetDeviceProxyManager();
 
-	// hide during PIE
-	FUIAction PlatformMenuShownDelegate;
-	if (FSlateApplication::IsInitialized())
-	{
-		PlatformMenuShownDelegate.CanExecuteAction = FCanExecuteAction::CreateRaw(&FSlateApplication::Get(), &FSlateApplication::IsNormalExecution);
-	}
-	PlatformMenuShownDelegate.IsActionVisibleDelegate = FIsActionButtonVisible::CreateLambda([]() 
-	{
-		return !FTurnkeyEditorSupport::IsPIERunning();
-	});
-
-	FToolMenuEntry Entry = FToolMenuEntry::InitComboButton(
-		"PlatformsMenu",
-		PlatformMenuShownDelegate,
-		FOnGetContent::CreateLambda([this] { return MakeTurnkeyMenuWidget(); }),
+	FToolMenuEntry& PlatformsEntry = MenuSection.AddSubMenu(
+		"Platforms",
 		LOCTEXT("PlatformMenu", "Platforms"),
 		LOCTEXT("PlatformMenu_Tooltip", "Platform related actions and settings (Launching, Packaging, custom builds, etc)"),
-		FSlateIcon(FAppStyle::Get().GetStyleSetName(), "PlayWorld.RepeatLastLaunch"), // not great name for a good "platforms" icon
-		false,
-		"PlatformsMenu");
-	Entry.StyleNameOverride = "CalloutToolbar";
+		FOnGetContent::CreateLambda([this]{ return MakeTurnkeyMenuWidget(); })
+	);
 
-	MenuSection.AddEntry(Entry);
+	PlatformsEntry.InsertPosition = InsertPosition;
 }
 
 // some shared functionality
diff --git a/Engine/Source/Editor/TurnkeySupport/Private/TurnkeySupportModule.h b/Engine/Source/Editor/TurnkeySupport/Private/TurnkeySupportModule.h
--- a/Engine/Source/Editor/TurnkeySupport/Private/TurnkeySupportModule.h
+++ b/Engine/Source/Editor/TurnkeySupport/Private/TurnkeySupportModule.h
@@ -21,7 +21,7 @@ class FTurnkeySupportModule	: public ITurnkeySupportModule
 {
 public:
 
-	virtual void MakeTurnkeyMenu(struct FToolMenuSection& MenuSection) const override;
+	virtual void MakeTurnkeyMenu(FToolMenuSection& MenuSection, const FToolMenuInsert& InsertPosition) const override;
 	virtual void MakeQuickLaunchItems(class UToolMenu* Menu, FOnQuickLaunchSelected ExternalOnClickDelegate) const override;
 	virtual void MakeSimulatorItems(class UToolMenu* Menu, FOnQuickLaunchSelected ExternalOnClickDelegate) const override;
 	virtual void RepeatQuickLaunch(FString DeviceId) override;
diff --git a/Engine/Source/Editor/TurnkeySupport/Public/Interfaces/ITurnkeySupportModule.h b/Engine/Source/Editor/TurnkeySupport/Public/Interfaces/ITurnkeySupportModule.h
--- a/Engine/Source/Editor/TurnkeySupport/Public/Interfaces/ITurnkeySupportModule.h
+++ b/Engine/Source/Editor/TurnkeySupport/Public/Interfaces/ITurnkeySupportModule.h
@@ -13,6 +13,8 @@
 
 #if UE_WITH_TURNKEY_SUPPORT
 
+struct FToolMenuInsert;
+struct FToolMenuSection;
 class SWidget;
 
 enum class ETurnkeyPlatformSdkStatus : uint8
@@ -68,9 +70,9 @@ class ITurnkeySupportModule
 public:
 
 	/**
-	 * Make a Platforms menu in the given MenuSection
+	 * Make a Platforms menu in the given MenuSection at the specified InsertPosition
 	 */
-	virtual void MakeTurnkeyMenu(struct FToolMenuSection& MenuSection) const = 0;
+	virtual void MakeTurnkeyMenu(FToolMenuSection& MenuSection, const FToolMenuInsert& InsertPosition) const = 0;
 
 	/**
 	 * Make menu items for Quick Launch, so they can be added to the Play menu
diff --git a/Engine/Source/Editor/UnrealEd/Private/Kismet2/DebuggerCommands.cpp b/Engine/Source/Editor/UnrealEd/Private/Kismet2/DebuggerCommands.cpp
--- a/Engine/Source/Editor/UnrealEd/Private/Kismet2/DebuggerCommands.cpp
+++ b/Engine/Source/Editor/UnrealEd/Private/Kismet2/DebuggerCommands.cpp
@@ -639,7 +639,7 @@ void FPlayWorldCommands::AddPIEPreviewDeviceActions(const FPlayWorldCommands &Co
 	}
 }
 
-void FPlayWorldCommands::BuildToolbar(FToolMenuSection& InSection, bool bIncludeLaunchButtonAndOptions)
+void FPlayWorldCommands::BuildToolbar(FToolMenuSection& InSection)
 {
 	FToolMenuEntry PlayMenuEntry =
 		FToolMenuEntry::InitToolBarButton(
@@ -756,8 +756,6 @@ void FPlayWorldCommands::BuildToolbar(FToolMenuSection& InSection, bool bInclude
 	InSection.AddEntry(PIEComboEntry);
 
 	// Single-stepping only buttons
-
-	InSection.AddSeparator(NAME_None).StyleNameOverride = FName("Toolbar.BackplateRight");
 	
 	FToolMenuEntry ShowCurrentStatementEntry = FToolMenuEntry::InitToolBarButton(FPlayWorldCommands::Get().ShowCurrentStatement, TAttribute<FText>(), TAttribute<FText>(), TAttribute<FSlateIcon>(), FName(TEXT("ShowCurrentStatement")));
 	ShowCurrentStatementEntry.StyleNameOverride = FName("Toolbar.BackplateLeft");
@@ -783,13 +781,6 @@ void FPlayWorldCommands::BuildToolbar(FToolMenuSection& InSection, bool bInclude
 	InSection.AddEntry(StepIntoEntry);	
 	InSection.AddEntry(StepOverEntry);	
 	InSection.AddEntry(StepOutEntry);
-
-#if UE_WITH_TURNKEY_SUPPORT
-	if (bIncludeLaunchButtonAndOptions)
-	{
-		ITurnkeySupportModule::Get().MakeTurnkeyMenu(InSection);
-	}
-#endif
 }
 
 static void MakePreviewDeviceMenu(FMenuBuilder& MenuBuilder)
diff --git a/Engine/Source/Editor/UnrealEd/Public/Kismet2/DebuggerCommands.h b/Engine/Source/Editor/UnrealEd/Public/Kismet2/DebuggerCommands.h
--- a/Engine/Source/Editor/UnrealEd/Public/Kismet2/DebuggerCommands.h
+++ b/Engine/Source/Editor/UnrealEd/Public/Kismet2/DebuggerCommands.h
@@ -56,7 +56,7 @@ public:
 	static UNREALED_API void BindGlobalPlayWorldCommands();
 
 	/** Populates a toolbar with the menu commands for play-world control (pause/resume/stop/possess/eject/step/show current loc) */
-	static UNREALED_API void BuildToolbar( FToolMenuSection& InSection, bool bIncludeLaunchButtonAndOptions = false );
+	static UNREALED_API void BuildToolbar(FToolMenuSection& InSection);
 
 	/**
 	* Return the active widget that processes play world actions for PIE
