[Editor] Implements the new design of the Platforms menu in the editor.

The only change in functionality is the Simulator section of the menu (that would appear under Quick Launch Game Target) has been removed. This was only used by iOS but we're dropping support for iOS simulators in 5.7.

There is no change to any other provided functionality. All previous functionality is still present, just moved around a bit.

Main Changes:
* Quick Launch has been moved into a sub-menu, rather than appearing in the root of the Platforms menu.
* Quick Launch was removed from the PIE menu when clicking "...".
* The actions that appeared in the Content Management section of each platform (Cook Content, Package Project, Prepare for Debugging, and custom builds) are now top-level entries with a sub-menu allowing selection of relevant options and a button to execute the action.
* A new top-level "SDK Management" entry contains the remaining functionality that previously appeared in the platform sub-menu, which is SDK installation and device updates.

New Behavior:
* When "Launch Configuration" in Editor Preferences is set to "Same as Editor", it uses the build configuration specified in the Quick Launch menu. It previously used the per-platform build configuration setting.
* Platforms that are not supported by the project (via Project Settings) are only shown in SDK Management. To cook or package these platforms, enable them in Project Settings.
* "Prepare for Debugging" is only shown if there is at least one supported platform that supports that operation. Currently those platforms are iOS and tvOS.
* Options shown under "Cook Content", "Package Project", etc., are separate for each action. The one exception is Flavor and Architecture selection remains global for that platform, so changing it for one menu changes for all menus.
* The tooltip for custom builds no longer shows the full command line. This is because we no longer know the platform or device when the menu is created since they are options that can be chosen.

diff --git a/Engine/Source/Developer/DeveloperToolSettings/Classes/Settings/PlatformsMenuSettings.h b/Engine/Source/Developer/DeveloperToolSettings/Classes/Settings/PlatformsMenuSettings.h
--- a/Engine/Source/Developer/DeveloperToolSettings/Classes/Settings/PlatformsMenuSettings.h
+++ b/Engine/Source/Developer/DeveloperToolSettings/Classes/Settings/PlatformsMenuSettings.h
@@ -10,7 +10,14 @@
 
 #define UE_API DEVELOPERTOOLSETTINGS_API
 
-
+// ProjectPackagingSettings has two arrays of custom builds: EngineCustomBuilds and ProjectCustomBuilds. That makes it
+// possible to have two builds with the same name (one in each list). For this reason, we need to store the settings
+// separately for both.
+enum class ECustomBuildType : uint8
+{
+	EngineCustomBuilds,
+	ProjectCustomBuilds,
+};
 
 UCLASS(MinimalAPI, config=Game)
 class UPlatformsMenuSettings
@@ -26,31 +33,62 @@ class UPlatformsMenuSettings
 	UPROPERTY(config)
 	FString LaunchOnTarget;
 
+	/** Name of the platform (IniPlatformName) to use when the "Cook Content" button is clicked in the menu */
+	UPROPERTY(config)
+	FName CookPlatform;
+
+	/** Build target to use when the "Cook Content" button is clicked in the menu */
+	UPROPERTY(config)
+	FString CookBuildTarget;
+
+	/** Name of the platform (IniPlatformName) to use when the "Package Project" button is clicked in the menu */
+	UPROPERTY(config)
+	FName PackagePlatform;
+
+	/** Build target to use when the "Package Project" button is clicked in the menu */
+	UPROPERTY(config)
+	FString PackageBuildTarget;
+
+	/** Name of the platform (IniPlatformName) to use when the "Prepare for Debugging" button is clicked in the menu */
+	UPROPERTY(config)
+	FName PrepareForDebuggingPlatform;
+
+	/** Build target to use when the "Prepare for Debugging" button is clicked in the menu */
+	UPROPERTY(config)
+	FString PrepareForDebuggingBuildTarget;
+
 	/** Gets the current launch on target, checking that it's valid, and the default build target if it is not */
 	UE_API const FTargetInfo* GetLaunchOnTargetInfo() const;
 	
-	/**
-	 * Get and set the per-platform build config and targetplatform settings for the Turnkey/Launch on menu
-	 */
-	UE_API EProjectPackagingBuildConfigurations GetBuildConfigurationForPlatform(FName PlatformName) const;
-	UE_API void SetBuildConfigurationForPlatform(FName PlatformName, EProjectPackagingBuildConfigurations Configuration);
+	UE_API FName GetTargetFlavorForPlatform(FName IniPlatformName) const;
+	UE_API void SetTargetFlavorForPlatform(FName IniPlatformName, FName TargetFlavorName);
 
-	UE_API FName GetTargetFlavorForPlatform(FName PlatformName) const;
-	UE_API void SetTargetFlavorForPlatform(FName PlatformName, FName TargetFlavorName);
+	UE_API FString GetArchitectureForPlatform(FName IniPlatformName) const;
+	UE_API void SetArchitectureForPlatform(FName IniPlatformName, FString ArchitectureName);
 
-	UE_API FString GetArchitectureForPlatform(FName PlatformName) const;
-	UE_API void SetArchitectureForPlatform(FName PlatformName, FString ArchitectureName);
+	UE_API const FTargetInfo* GetBuildTargetInfo(const FString& TargetName, bool& bOutIsProjectTarget) const;
 
-	UE_API FString GetBuildTargetForPlatform(FName PlatformName) const;
-	UE_API void SetBuildTargetForPlatform(FName PlatformName, FString BuildTargetName);
+	UE_API EProjectPackagingBuildConfigurations GetLaunchOnBuildConfiguration() const;
+	UE_API void SetLaunchOnBuildConfiguration(EProjectPackagingBuildConfigurations BuildConfiguration);
 
-	UE_API const FTargetInfo* GetBuildTargetInfoForPlatform(FName PlatformName, bool& bOutIsProjectTarget) const;
+	UE_API EProjectPackagingBuildConfigurations GetPackageBuildConfiguration() const;
+	UE_API void SetPackageBuildConfiguration(EProjectPackagingBuildConfigurations BuildConfiguration);
+
+	/* Returns the IniPlatformName selected for the custom build */
+	UE_API FName GetPlatformForCustomBuild(ECustomBuildType Type, const FString& CustomBuildName) const;
+	UE_API void SetPlatformForCustomBuild(ECustomBuildType Type, const FString& CustomBuildName, FName IniPlatformName);
+
+	UE_API FString GetDeviceIdForCustomBuild(ECustomBuildType Type, const FString& CustomBuildName) const;
+	UE_API void SetDeviceIdForCustomBuild(ECustomBuildType Type, const FString& CustomBuildName, const FString& DeviceId);
+
+	UE_API EProjectPackagingBuildConfigurations GetBuildConfigurationForCustomBuild(ECustomBuildType Type, const FString& CustomBuildName) const;
+	UE_API void SetBuildConfigurationForCustomBuild(ECustomBuildType Type, const FString& CustomBuildName, EProjectPackagingBuildConfigurations BuildConfiguration);
+
+	UE_API FString GetBuildTargetForCustomBuild(ECustomBuildType Type, const FString& CustomBuildName) const;
+	UE_API void SetBuildTargetForCustomBuild(ECustomBuildType Type, const FString& CustomBuildName, const FString& TargetName);
 
-private:
-	/** Per platform build configuration */
-	UPROPERTY(config)
-	TMap<FName, EProjectPackagingBuildConfigurations> PerPlatformBuildConfig;
 
+private:
 	/** Per platform flavor cooking target */
 	UPROPERTY(config)
 	TMap<FName, FName> PerPlatformTargetFlavorName;
@@ -59,12 +97,45 @@ private:
 	UPROPERTY(config, EditAnywhere, Category=Project)
 	TMap<FName, FString> PerPlatformArchitecture;
 
-	/** Per platform build target */
-	UPROPERTY(config, EditAnywhere, Category=Project)
-	TMap<FName, FString> PerPlatformBuildTarget;
+	/** Platform name setting stored per entry in EngineCustomBuilds */
+	UPROPERTY(config)
+	TMap<FString, FName> PerEngineCustomBuildPlatformName;
 
+	/** Platform name setting stored per entry in ProjectCustomBuilds */
+	UPROPERTY(config)
+	TMap<FString, FName> PerProjectCustomBuildPlatformName;
 
-	
+	/** Device ID setting stored per entry in EngineCustomBuilds */
+	UPROPERTY(config)
+	TMap<FString, FString> PerEngineCustomBuildDeviceId;
+
+	/** Device ID setting stored per entry in ProjectCustomBuilds */
+	UPROPERTY(config)
+	TMap<FString, FString> PerProjectCustomBuildDeviceId;
+
+	/** Build configuration setting stored per entry in EngineCustomBuilds */
+	UPROPERTY(config)
+	TMap<FString, EProjectPackagingBuildConfigurations> PerEngineCustomBuildBuildConfiguration;
+
+	/** Build configuration setting stored per entry in ProjectCustomBuilds */
+	UPROPERTY(config)
+	TMap<FString, EProjectPackagingBuildConfigurations> PerProjectCustomBuildBuildConfiguration;
+
+	/** Build target name setting stored per entry in EngineCustomBuilds */
+	UPROPERTY(config)
+	TMap<FString, FString> PerEngineCustomBuildBuildTargetName;
+
+	/** Build target name setting stored per entry in ProjectCustomBuilds */
+	UPROPERTY(config)
+	TMap<FString, FString> PerProjectCustomBuildBuildTargetName;
+
+	/** Build configuration to use for LaunchOn */
+	UPROPERTY(config)
+	TOptional<EProjectPackagingBuildConfigurations> LaunchOnBuildConfiguration;
+
+	/** Build configuration to use when the "Package Project" button is clicked in the menu */
+	UPROPERTY(config)
+	TOptional<EProjectPackagingBuildConfigurations> PackageBuildConfiguration;
 };
 
 #undef UE_API
diff --git a/Engine/Source/Developer/DeveloperToolSettings/Private/PlatformsMenuSettings.cpp b/Engine/Source/Developer/DeveloperToolSettings/Private/PlatformsMenuSettings.cpp
--- a/Engine/Source/Developer/DeveloperToolSettings/Private/PlatformsMenuSettings.cpp
+++ b/Engine/Source/Developer/DeveloperToolSettings/Private/PlatformsMenuSettings.cpp
@@ -21,78 +21,220 @@ UPlatformsMenuSettings::UPlatformsMenuSettings(const FObjectInitializer& ObjectI
 }
 
 
-const FTargetInfo* UPlatformsMenuSettings::GetBuildTargetInfoForPlatform(FName PlatformName, bool& bOutIsProjectTarget) const
+const FTargetInfo* UPlatformsMenuSettings::GetBuildTargetInfo(const FString& TargetName, bool& bOutIsProjectTarget) const
 {
-	return FindBestTargetInfo(GetBuildTargetForPlatform(PlatformName), true, &bOutIsProjectTarget);
+	return FindBestTargetInfo(TargetName, true, &bOutIsProjectTarget);
 }
 
-EProjectPackagingBuildConfigurations UPlatformsMenuSettings::GetBuildConfigurationForPlatform(FName PlatformName) const
+FName UPlatformsMenuSettings::GetTargetFlavorForPlatform(FName IniPlatformName) const
 {
-	const EProjectPackagingBuildConfigurations* Value = PerPlatformBuildConfig.Find(PlatformName);
+	const FName* Value = PerPlatformTargetFlavorName.Find(IniPlatformName);
 
-	// PPBC_MAX defines the default project setting case and should be handled accordingly.
-	return Value == nullptr ? EProjectPackagingBuildConfigurations::PPBC_MAX : *Value;
+	// the flavor name is also the name of the vanilla info
+	return Value == nullptr ? IniPlatformName : *Value;
+}
+
+void UPlatformsMenuSettings::SetTargetFlavorForPlatform(FName IniPlatformName, FName TargetFlavorName)
+{
+	PerPlatformTargetFlavorName.Add(IniPlatformName, TargetFlavorName);
+}
+
+FString UPlatformsMenuSettings::GetArchitectureForPlatform(FName IniPlatformName) const
+{
+	const FString* Value = PerPlatformArchitecture.Find(IniPlatformName);
+	return Value == nullptr ? TEXT("") : *Value;
+}
+
+void UPlatformsMenuSettings::SetArchitectureForPlatform(FName IniPlatformName, FString ArchitectureName)
+{
+	PerPlatformArchitecture.Add(IniPlatformName, ArchitectureName);
+}
+
+const FTargetInfo* UPlatformsMenuSettings::GetLaunchOnTargetInfo() const
+{
+	return FindBestTargetInfo(LaunchOnTarget, true, nullptr);
+}
+
+EProjectPackagingBuildConfigurations UPlatformsMenuSettings::GetLaunchOnBuildConfiguration() const
+{
+	// PPBC_MAX is used to indicate use the project default.
+	return LaunchOnBuildConfiguration.IsSet() ? LaunchOnBuildConfiguration.GetValue() : EProjectPackagingBuildConfigurations::PPBC_MAX;
 }
 
-void UPlatformsMenuSettings::SetBuildConfigurationForPlatform(FName PlatformName, EProjectPackagingBuildConfigurations Configuration)
+void UPlatformsMenuSettings::SetLaunchOnBuildConfiguration(EProjectPackagingBuildConfigurations BuildConfiguration)
 {
-	if (Configuration == EProjectPackagingBuildConfigurations::PPBC_MAX)
+	// PPBC_MAX is a special value to indicate use the project default.
+	if (BuildConfiguration == EProjectPackagingBuildConfigurations::PPBC_MAX)
 	{
-		PerPlatformBuildConfig.Remove(PlatformName);
+		LaunchOnBuildConfiguration.Reset();
 	}
 	else
 	{
-		PerPlatformBuildConfig.Add(PlatformName, Configuration);
+		LaunchOnBuildConfiguration = BuildConfiguration;
 	}
 }
 
-FName UPlatformsMenuSettings::GetTargetFlavorForPlatform(FName FlavorName) const
+EProjectPackagingBuildConfigurations UPlatformsMenuSettings::GetPackageBuildConfiguration() const
 {
-	const FName* Value = PerPlatformTargetFlavorName.Find(FlavorName);
+	// PPBC_MAX is used to indicate use the project default.
+	return PackageBuildConfiguration.IsSet() ? PackageBuildConfiguration.GetValue() : EProjectPackagingBuildConfigurations::PPBC_MAX;
+}
 
-	// the flavor name is also the name of the vanilla info
-	return Value == nullptr ? FlavorName : *Value;
+void UPlatformsMenuSettings::SetPackageBuildConfiguration(EProjectPackagingBuildConfigurations BuildConfiguration)
+{
+	// PPBC_MAX is a special value to indicate use the project default.
+	if (BuildConfiguration == EProjectPackagingBuildConfigurations::PPBC_MAX)
+	{
+		PackageBuildConfiguration.Reset();
+	}
+	else
+	{
+		PackageBuildConfiguration = BuildConfiguration;
+	}
 }
 
-void UPlatformsMenuSettings::SetTargetFlavorForPlatform(FName PlatformName, FName TargetFlavorName)
+FName UPlatformsMenuSettings::GetPlatformForCustomBuild(ECustomBuildType Type, const FString& CustomBuildName) const
 {
-	PerPlatformTargetFlavorName.Add(PlatformName, TargetFlavorName);
+	const FName* IniPlatformName = nullptr;
+
+	if (Type == ECustomBuildType::EngineCustomBuilds)
+	{
+		IniPlatformName = PerEngineCustomBuildPlatformName.Find(CustomBuildName);
+	}
+	else if (Type == ECustomBuildType::ProjectCustomBuilds)
+	{
+		IniPlatformName = PerProjectCustomBuildPlatformName.Find(CustomBuildName);
+	}
+	else
+	{
+		checkNoEntry();
+	}
+
+	return IniPlatformName == nullptr ? NAME_None : *IniPlatformName;
 }
 
-FString UPlatformsMenuSettings::GetArchitectureForPlatform(FName PlatformName) const
+void UPlatformsMenuSettings::SetPlatformForCustomBuild(ECustomBuildType Type, const FString& CustomBuildName, FName IniPlatformName)
 {
-	const FString* Value = PerPlatformArchitecture.Find(PlatformName);
-	return Value == nullptr ? TEXT("") : *Value;
+	if (Type == ECustomBuildType::EngineCustomBuilds)
+	{
+		PerEngineCustomBuildPlatformName.Add(CustomBuildName, IniPlatformName);
+	}
+	else if (Type == ECustomBuildType::ProjectCustomBuilds)
+	{
+		PerProjectCustomBuildPlatformName.Add(CustomBuildName, IniPlatformName);
+	}
+	else
+	{
+		checkNoEntry();
+	}
+}
+
+FString UPlatformsMenuSettings::GetDeviceIdForCustomBuild(ECustomBuildType Type, const FString& CustomBuildName) const
+{
+	const FString* DeviceId = nullptr;
+
+	if (Type == ECustomBuildType::EngineCustomBuilds)
+	{
+		DeviceId = PerEngineCustomBuildDeviceId.Find(CustomBuildName);
+	}
+	else if (Type == ECustomBuildType::ProjectCustomBuilds)
+	{
+		DeviceId = PerProjectCustomBuildDeviceId.Find(CustomBuildName);
+	}
+	else
+	{
+		checkNoEntry();
+	}
+
+	return DeviceId == nullptr ? FString() : *DeviceId;
 }
 
-void UPlatformsMenuSettings::SetArchitectureForPlatform(FName PlatformName, FString ArchitectureName)
+void UPlatformsMenuSettings::SetDeviceIdForCustomBuild(ECustomBuildType Type, const FString& CustomBuildName, const FString& DeviceId)
 {
-	PerPlatformArchitecture.Add(PlatformName, ArchitectureName);
+	if (Type == ECustomBuildType::EngineCustomBuilds)
+	{
+		PerEngineCustomBuildDeviceId.Add(CustomBuildName, DeviceId);
+	}
+	else if (Type == ECustomBuildType::ProjectCustomBuilds)
+	{
+		PerProjectCustomBuildDeviceId.Add(CustomBuildName, DeviceId);
+	}
+	else
+	{
+		checkNoEntry();
+	}
 }
 
-FString UPlatformsMenuSettings::GetBuildTargetForPlatform(FName PlatformName) const
+EProjectPackagingBuildConfigurations UPlatformsMenuSettings::GetBuildConfigurationForCustomBuild(ECustomBuildType Type, const FString& CustomBuildName) const
 {
-	const FString* Value = PerPlatformBuildTarget.Find(PlatformName);
+	const EProjectPackagingBuildConfigurations* BuildConfiguration = nullptr;
+
+	if (Type == ECustomBuildType::EngineCustomBuilds)
+	{
+		BuildConfiguration = PerEngineCustomBuildBuildConfiguration.Find(CustomBuildName);
+	}
+	else if (Type == ECustomBuildType::ProjectCustomBuilds)
+	{
+		BuildConfiguration = PerProjectCustomBuildBuildConfiguration.Find(CustomBuildName);
+	}
+	else
+	{
+		checkNoEntry();
+	}
 
-	// empty string defines the default project setting case and should be handled accordingly.
-	return Value == nullptr ? "" : *Value;
+	return BuildConfiguration == nullptr ? EProjectPackagingBuildConfigurations::PPBC_MAX : *BuildConfiguration;
 }
 
-void UPlatformsMenuSettings::SetBuildTargetForPlatform(FName PlatformName, FString BuildTargetName)
+void UPlatformsMenuSettings::SetBuildConfigurationForCustomBuild(ECustomBuildType Type, const FString& CustomBuildName, EProjectPackagingBuildConfigurations BuildConfiguration)
 {
-	if (BuildTargetName.IsEmpty())
+	if (Type == ECustomBuildType::EngineCustomBuilds)
 	{
-		PerPlatformBuildTarget.Remove(PlatformName);
+		PerEngineCustomBuildBuildConfiguration.Add(CustomBuildName, BuildConfiguration);
+	}
+	else if (Type == ECustomBuildType::ProjectCustomBuilds)
+	{
+		PerProjectCustomBuildBuildConfiguration.Add(CustomBuildName, BuildConfiguration);
 	}
 	else
 	{
-		PerPlatformBuildTarget.Add(PlatformName, BuildTargetName);
+		checkNoEntry();
 	}
 }
 
-const FTargetInfo* UPlatformsMenuSettings::GetLaunchOnTargetInfo() const
+FString UPlatformsMenuSettings::GetBuildTargetForCustomBuild(ECustomBuildType Type, const FString& CustomBuildName) const
 {
-	return FindBestTargetInfo(LaunchOnTarget, true, nullptr);
+	const FString* TargetName = nullptr;
+
+	if (Type == ECustomBuildType::EngineCustomBuilds)
+	{
+		TargetName = PerEngineCustomBuildBuildTargetName.Find(CustomBuildName);
+	}
+	else if (Type == ECustomBuildType::ProjectCustomBuilds)
+	{
+		TargetName = PerProjectCustomBuildBuildTargetName.Find(CustomBuildName);
+	}
+	else
+	{
+		checkNoEntry();
+	}
+
+	return TargetName == nullptr ? FString() : *TargetName;
+}
+
+void UPlatformsMenuSettings::SetBuildTargetForCustomBuild(ECustomBuildType Type, const FString& CustomBuildName, const FString& TargetName)
+{
+	if (Type == ECustomBuildType::EngineCustomBuilds)
+	{
+		PerEngineCustomBuildBuildTargetName.Add(CustomBuildName, TargetName);
+	}
+	else if (Type == ECustomBuildType::ProjectCustomBuilds)
+	{
+		PerProjectCustomBuildBuildTargetName.Add(CustomBuildName, TargetName);
+	}
+	else
+	{
+		checkNoEntry();
+	}
 }
 
 #undef LOCTEXT_NAMESPACE
diff --git a/Engine/Source/Editor/EditorStyle/Private/StarshipStyle.cpp b/Engine/Source/Editor/EditorStyle/Private/StarshipStyle.cpp
--- a/Engine/Source/Editor/EditorStyle/Private/StarshipStyle.cpp
+++ b/Engine/Source/Editor/EditorStyle/Private/StarshipStyle.cpp
@@ -1354,6 +1354,7 @@ void FStarshipEditorStyle::FStyle::SetupGeneralStyles()
 		Set( "MainFrame.RefreshVisualStudio",        new IMAGE_BRUSH_SVG( "Starship/Common/RefreshVisualStudio", Icon16x16 ) );
 		Set( "MainFrame.OpenSourceCodeEditor",       new IMAGE_BRUSH_SVG( "Starship/Common/SourceCodeEditor", Icon16x16));
 		Set( "MainFrame.RefreshSourceCodeEditor",    new IMAGE_BRUSH_SVG( "Starship/Common/RefreshSourceCodeEditor", Icon16x16));
+		Set( "MainFrame.Project",                    new IMAGE_BRUSH_SVG( "Starship/Common/Project", Icon16x16 ) );
 		Set( "MainFrame.PackageProject",             new IMAGE_BRUSH_SVG( "Starship/Common/ProjectPackage", Icon16x16 ) );
 		Set( "MainFrame.RecentProjects",             new IMAGE_BRUSH_SVG( "Starship/Common/ProjectsRecent", Icon16x16 ) );
 		Set( "MainFrame.RecentLevels",               new IMAGE_BRUSH_SVG( "Starship/Common/LevelRecent", Icon16x16 ) );
diff --git a/Engine/Source/Editor/TurnkeySupport/Private/TurnkeySupportModule.cpp b/Engine/Source/Editor/TurnkeySupport/Private/TurnkeySupportModule.cpp
--- a/Engine/Source/Editor/TurnkeySupport/Private/TurnkeySupportModule.cpp
+++ b/Engine/Source/Editor/TurnkeySupport/Private/TurnkeySupportModule.cpp
@@ -279,7 +279,7 @@ public:
 		return PackagingSettings;
 	}
 
-	static void CookOrPackage(FName IniPlatformName, EPrepareContentMode Mode)
+	static void CookOrPackage(FName IniPlatformName, const FString& BuildTargetName, EPrepareContentMode Mode)
 	{
 		TArray<FAnalyticsEventAttribute> AnalyticsParamArray;
 
@@ -341,7 +341,7 @@ public:
 		}
 
 		bool bIsProjectBuildTarget = false;
-		const FTargetInfo* BuildTargetInfo = PlatformsSettings->GetBuildTargetInfoForPlatform(IniPlatformName, bIsProjectBuildTarget);
+		const FTargetInfo* BuildTargetInfo = PlatformsSettings->GetBuildTargetInfo(BuildTargetName, bIsProjectBuildTarget);
 
 		// Only add the -Target=... argument for code projects. Content projects will return UnrealGame/UnrealClient/UnrealServer here, but
 		// may need a temporary target generated to enable/disable plugins. Specifying -Target in these cases will cause packaging to fail,
@@ -522,7 +522,7 @@ public:
 				BuildCookRunParams += FString::Printf(TEXT(" -manifests -createchunkinstall -chunkinstalldirectory=\"%s\" -chunkinstallversion=%s"), *(PackagingSettings->HttpChunkInstallDataDirectory.Path), *(PackagingSettings->HttpChunkInstallDataVersion));
 			}
 
-			EProjectPackagingBuildConfigurations BuildConfig = PlatformsSettings->GetBuildConfigurationForPlatform(IniPlatformName);
+			EProjectPackagingBuildConfigurations BuildConfig = PlatformsSettings->GetPackageBuildConfiguration();
 			// if PPBC_MAX is set, then the project default should be used instead of the per platform build config
 			if (BuildConfig == EProjectPackagingBuildConfigurations::PPBC_MAX)
 			{
@@ -619,17 +619,91 @@ public:
 		FTurnkeyEditorSupport::RunUAT(CommandLine, PlatformInfo->DisplayName, ContentPrepDescription, ContentPrepTaskName, ContentPrepIcon, &AnalyticsParamArray);
 	}
 
-	static bool CanExecuteCustomBuild(FName IniPlatformName, FProjectBuildSettings Build)
+	static void SetPlatformForCustomBuild(ECustomBuildType Type, FString CustomBuildName, FName IniPlatformName)
 	{
+		UPlatformsMenuSettings* PlatformsSettings = GetMutableDefault<UPlatformsMenuSettings>();
+		PlatformsSettings->SetPlatformForCustomBuild(Type, CustomBuildName, IniPlatformName);
+		PlatformsSettings->SaveConfig();
+	}
+
+	static bool PlatformForCustomBuildIsChecked(ECustomBuildType Type, FString CustomBuildName, FName IniPlatformName)
+	{
+		const UPlatformsMenuSettings* PlatformsSettings = GetDefault<UPlatformsMenuSettings>();
+		FName CheckedPlatform = PlatformsSettings->GetPlatformForCustomBuild(Type, CustomBuildName);
+		return CheckedPlatform == IniPlatformName;
+	}
+
+	// Per-device custom builds show a list of devices rather than platforms. Since we know the platform associated with
+	// each device, we set both at once to ensure the platform is set.
+	static void SetDeviceIdForCustomBuild(ECustomBuildType Type, FString CustomBuildName, FName IniPlatformName, FString DeviceId)
+	{
+		UPlatformsMenuSettings* PlatformsSettings = GetMutableDefault<UPlatformsMenuSettings>();
+		PlatformsSettings->SetPlatformForCustomBuild(Type, CustomBuildName, IniPlatformName);
+		PlatformsSettings->SetDeviceIdForCustomBuild(Type, CustomBuildName, DeviceId);
+		PlatformsSettings->SaveConfig();
+	}
+
+	static bool DeviceIdForCustomBuildIsChecked(ECustomBuildType Type, FString CustomBuildName, FString DeviceId)
+	{
+		const UPlatformsMenuSettings* PlatformsSettings = GetDefault<UPlatformsMenuSettings>();
+		FString CheckedDeviceId = PlatformsSettings->GetDeviceIdForCustomBuild(Type, CustomBuildName);
+		return CheckedDeviceId == DeviceId;
+	}
+
+	static void SetBuildConfigurationForCustomBuild(ECustomBuildType Type, FString CustomBuildName, EProjectPackagingBuildConfigurations BuildConfiguration)
+	{
+		UPlatformsMenuSettings* PlatformsSettings = GetMutableDefault<UPlatformsMenuSettings>();
+		PlatformsSettings->SetBuildConfigurationForCustomBuild(Type, CustomBuildName, BuildConfiguration);
+		PlatformsSettings->SaveConfig();
+	}
+
+	static bool BuildConfigurationForCustomBuildIsChecked(ECustomBuildType Type, FString CustomBuildName, EProjectPackagingBuildConfigurations BuildConfiguration)
+	{
+		const UPlatformsMenuSettings* PlatformsSettings = GetDefault<UPlatformsMenuSettings>();
+		EProjectPackagingBuildConfigurations CheckedBuildConfiguration = PlatformsSettings->GetBuildConfigurationForCustomBuild(Type, CustomBuildName);
+		return CheckedBuildConfiguration == BuildConfiguration;
+	}
+
+	static void SetBuildTargetForCustomBuild(ECustomBuildType Type, FString CustomBuildName, FString TargetName)
+	{
+		UPlatformsMenuSettings* PlatformsSettings = GetMutableDefault<UPlatformsMenuSettings>();
+		PlatformsSettings->SetBuildTargetForCustomBuild(Type, CustomBuildName, TargetName);
+		PlatformsSettings->SaveConfig();
+	}
+
+	static bool BuildTargetForCustomBuildIsChecked(ECustomBuildType Type, FString CustomBuildName, FString TargetName)
+	{
+		const UPlatformsMenuSettings* PlatformsSettings = GetDefault<UPlatformsMenuSettings>();
+		FString CheckedTargetName = PlatformsSettings->GetBuildTargetForCustomBuild(Type, CustomBuildName);
+		return CheckedTargetName == TargetName;
+	}
+
+	static bool CanExecuteCustomBuild(ECustomBuildType Type, FString BuildName, bool bNeedsDevice)
+	{
+		// Read the settings associated with this build
+		const UPlatformsMenuSettings* MenuSettings = GetDefault<UPlatformsMenuSettings>();
+
+		FName IniPlatformName = MenuSettings->GetPlatformForCustomBuild(Type, BuildName);
+
 		if (GetTargetPlatformManager()->FindTargetPlatform(IniPlatformName.ToString()) == nullptr)
 		{
 			return false;
 		}
 
+		if (bNeedsDevice)
+		{
+			FString DeviceId = MenuSettings->GetDeviceIdForCustomBuild(Type, BuildName);
+
+			if (DeviceId.IsEmpty())
+			{
+				return false;
+			}
+		}
+
 		return true;
 	}
 
-	static FString GetCustomBuildCommandLine(FName IniPlatformName, const FString& DeviceId, const FProjectBuildSettings& Build)
+	static FString GetCustomBuildCommandLine(FName IniPlatformName, EProjectPackagingBuildConfigurations BuildConfigSetting, const FString& BuildTargetName, const FString& DeviceId, const FProjectBuildSettings& Build)
 	{
 		const FString ProjectPath = GetProjectPathForTurnkey();
 
@@ -649,14 +723,14 @@ public:
 			UProjectPackagingSettings* PlatformPackagingSettings = FTurnkeySupportCallbacks::GetPackagingSettingsForPlatform(IniPlatformName);
 			const UPlatformsMenuSettings* PlatformsSettings = GetDefault<UPlatformsMenuSettings>();
 			bool bIsProjectBuildTarget = false;
-			const FTargetInfo* BuildTargetInfo = PlatformsSettings->GetBuildTargetInfoForPlatform(IniPlatformName, bIsProjectBuildTarget);
+			const FTargetInfo* BuildTargetInfo = PlatformsSettings->GetBuildTargetInfo(BuildTargetName, bIsProjectBuildTarget);
 
 			CommandLine.Appendf(TEXT(" -overridetarget=%s"), *BuildTargetInfo->Name);
 
 			// distribution builds can set shipping 
 			EProjectPackagingBuildConfigurations BuildConfig = PlatformPackagingSettings->ForDistribution ?
 				EProjectPackagingBuildConfigurations::PPBC_Shipping :
-			PlatformsSettings->GetBuildConfigurationForPlatform(IniPlatformName);
+				BuildConfigSetting;
 
 			// if PPBC_MAX is set, then the project default should be used instead of the per platform build config
 			if (BuildConfig == EProjectPackagingBuildConfigurations::PPBC_MAX)
@@ -678,8 +752,16 @@ public:
 		return CommandLine;
 	}
 
-	static void ExecuteCustomBuild(FName IniPlatformName, FString DeviceId, FProjectBuildSettings Build)
+	static void ExecuteCustomBuild(ECustomBuildType Type, FProjectBuildSettings Build)
 	{
+		// Read the settings associated with this build
+		const UPlatformsMenuSettings* MenuSettings = GetDefault<UPlatformsMenuSettings>();
+
+		FName IniPlatformName = MenuSettings->GetPlatformForCustomBuild(Type, Build.Name);
+		FString DeviceId = MenuSettings->GetDeviceIdForCustomBuild(Type, Build.Name);
+		FString BuildTargetName = MenuSettings->GetBuildTargetForCustomBuild(Type, Build.Name);
+		EProjectPackagingBuildConfigurations BuildConfig = MenuSettings->GetBuildConfigurationForCustomBuild(Type, Build.Name);
+
 		const PlatformInfo::FTargetPlatformInfo* PlatformInfo = PlatformInfo::FindPlatformInfo(GetDefault<UPlatformsMenuSettings>()->GetTargetFlavorForPlatform(IniPlatformName));
 		const FString ProjectPath = GetProjectPathForTurnkey();
 
@@ -690,7 +772,7 @@ public:
 		{
 			CommandLine.Appendf(TEXT(" -ScriptsForProject=\"%s\" "), *ProjectPath);
 		}
-		CommandLine += GetCustomBuildCommandLine(IniPlatformName, DeviceId, Build);
+		CommandLine += GetCustomBuildCommandLine(IniPlatformName, BuildConfig, BuildTargetName, DeviceId, Build);
 
 #if ALLOW_CONTROL_TO_COPY_COMMANDLINE
 		bool bIsControlHeld = FSlateApplication::Get().GetModifierKeys().IsControlDown();
@@ -701,7 +783,7 @@ public:
 			FString LogFilename, ReportFilename;
 			CommandLine += GetLogAndReportCommandline(LogFilename, ReportFilename);
 
-			FTurnkeyEditorSupport::RunUAT(CommandLine, PlatformInfo->DisplayName, LOCTEXT("Turnkey_GettingCommandLine", "Copying Commandline"), LOCTEXT("Turnkey_CustomTaskNameCmdLine", "CustomCommandLine"), nullptr /* TaskIcon */,
+			FTurnkeyEditorSupport::RunUAT(CommandLine, PlatformInfo->DisplayName, LOCTEXT("Turnkey_GettingCommandLine", "Copying Commandline"), LOCTEXT("Turnkey_CustomTaskNameCmdLine", "CustomCommandLine"), nullptr /* TaskIcon */, nullptr /* OptionalAnalyticsParamArray */,
 				[BuildName=Build.Name, ReportFilename](FString, double)
 				{
 					FString ReportContents;
@@ -747,17 +829,18 @@ public:
 		}
 	}
 
-	static void SetPackageBuildConfiguration(const PlatformInfo::FTargetPlatformInfo* Info, EProjectPackagingBuildConfigurations BuildConfiguration)
+	static void SetPackageBuildConfiguration(EProjectPackagingBuildConfigurations BuildConfiguration)
 	{
 		UPlatformsMenuSettings* PlatformsSettings = GetMutableDefault<UPlatformsMenuSettings>();
-		PlatformsSettings->SetBuildConfigurationForPlatform(Info->IniPlatformName, BuildConfiguration);
+		PlatformsSettings->SetPackageBuildConfiguration(BuildConfiguration);
 		PlatformsSettings->SaveConfig();
 	}
 
-	static bool PackageBuildConfigurationIsChecked(const PlatformInfo::FTargetPlatformInfo* Info, EProjectPackagingBuildConfigurations BuildConfiguration)
+	static bool PackageBuildConfigurationIsChecked(EProjectPackagingBuildConfigurations BuildConfiguration)
 	{
-		return GetDefault<UPlatformsMenuSettings>()->GetBuildConfigurationForPlatform(Info->IniPlatformName) == BuildConfiguration;
-	}	
+		EProjectPackagingBuildConfigurations CheckedBuildConfiguration = GetDefault<UPlatformsMenuSettings>()->GetPackageBuildConfiguration();
+		return CheckedBuildConfiguration == BuildConfiguration;
+	}
 	
 	static void SetActiveFlavor(const PlatformInfo::FTargetPlatformInfo* Info)
 	{
@@ -766,49 +849,52 @@ public:
 		PlatformsSettings->SaveConfig();
 	}
 
-	static bool CanSetActiveFlavor(const PlatformInfo::FTargetPlatformInfo* Info)
-	{
-		return true;
-	}
-
 	static bool SetActiveFlavorIsChecked(const PlatformInfo::FTargetPlatformInfo* Info)
 	{
 		return GetDefault<UPlatformsMenuSettings>()->GetTargetFlavorForPlatform(Info->IniPlatformName) == Info->Name;
 	}
 
-	static void SetActiveArchitecture(const PlatformInfo::FTargetPlatformInfo* VanillaInfo, FString Architecture)
+	static void SetActiveArchitecture(FName IniPlatformName, FString Architecture)
 	{
 		UPlatformsMenuSettings* PlatformsSettings = GetMutableDefault<UPlatformsMenuSettings>();
-		PlatformsSettings->SetArchitectureForPlatform(VanillaInfo->IniPlatformName, Architecture);
+		PlatformsSettings->SetArchitectureForPlatform(IniPlatformName, Architecture);
 		PlatformsSettings->SaveConfig();
 	}
 
-	static bool CanSetActiveArchitecture(const PlatformInfo::FTargetPlatformInfo* Info, FString Architecture)
+	static bool SetActiveArchitectureIsChecked(FName IniPlatformName, FString Architecture)
 	{
-		return true;
+		return GetDefault<UPlatformsMenuSettings>()->GetArchitectureForPlatform(IniPlatformName) == Architecture;
 	}
 
-	static bool SetActiveArchitectureIsChecked(const PlatformInfo::FTargetPlatformInfo* Info, FString Architecture)
+	static void SetPackageBuildTarget(FString TargetName)
 	{
-		return GetDefault<UPlatformsMenuSettings>()->GetArchitectureForPlatform(Info->IniPlatformName) == Architecture;
+		UPlatformsMenuSettings* PlatformsSettings = GetMutableDefault<UPlatformsMenuSettings>();
+		PlatformsSettings->PackageBuildTarget = MoveTemp(TargetName);
+		PlatformsSettings->SaveConfig();
+	}
+
+	static bool PackageBuildTargetIsChecked(FString TargetName)
+	{
+		return GetDefault<UPlatformsMenuSettings>()->PackageBuildTarget == TargetName;
 	}
 
-	static void SetPackageBuildTarget(const PlatformInfo::FTargetPlatformInfo* Info, FString TargetName)
+	static void SetLaunchOnBuildConfiguration(EProjectPackagingBuildConfigurations BuildConfiguration)
 	{
 		UPlatformsMenuSettings* PlatformsSettings = GetMutableDefault<UPlatformsMenuSettings>();
-		PlatformsSettings->SetBuildTargetForPlatform(Info->IniPlatformName, TargetName);
+		PlatformsSettings->SetLaunchOnBuildConfiguration(BuildConfiguration);
 		PlatformsSettings->SaveConfig();
 	}
 
-	static bool PackageBuildTargetIsChecked(const PlatformInfo::FTargetPlatformInfo* Info, FString TargetName)
+	static bool LaunchOnBuildConfigurationIsChecked(EProjectPackagingBuildConfigurations BuildConfiguration)
 	{
-		return GetDefault<UPlatformsMenuSettings>()->GetBuildTargetForPlatform(Info->IniPlatformName) == TargetName;
+		EProjectPackagingBuildConfigurations CheckedBuildConfiguration = GetDefault<UPlatformsMenuSettings>()->GetLaunchOnBuildConfiguration();
+		return CheckedBuildConfiguration == BuildConfiguration;
 	}
 
 	static void SetLaunchOnBuildTarget(FString TargetName)
 	{
 		UPlatformsMenuSettings* PlatformsSettings = GetMutableDefault<UPlatformsMenuSettings>();
-		PlatformsSettings->LaunchOnTarget = TargetName;
+		PlatformsSettings->LaunchOnTarget = MoveTemp(TargetName);
 		PlatformsSettings->SaveConfig();
 	}
 
@@ -857,14 +943,102 @@ public:
 
 	}
 
-	static bool CanSetCookOnTheFly()
+	static bool SetCookOnTheFlyIsChecked()
 	{
-		return true;
+		return GetDefault<UCookerSettings>()->bCookOnTheFlyForLaunchOn;
 	}
 
-	static bool SetCookOnTheFlyIsChecked()
+	static void SetCookPlatform(FName IniPlatformName)
 	{
-		return GetDefault<UCookerSettings>()->bCookOnTheFlyForLaunchOn;
+		UPlatformsMenuSettings* PlatformsSettings = GetMutableDefault<UPlatformsMenuSettings>();
+		PlatformsSettings->CookPlatform = IniPlatformName;
+		PlatformsSettings->SaveConfig();
+	}
+
+	static bool CookPlatformIsChecked(FName IniPlatformName)
+	{
+		return GetDefault<UPlatformsMenuSettings>()->CookPlatform == IniPlatformName;
+	}
+
+	static void SetCookBuildTarget(FString TargetName)
+	{
+		UPlatformsMenuSettings* PlatformsSettings = GetMutableDefault<UPlatformsMenuSettings>();
+		PlatformsSettings->CookBuildTarget = MoveTemp(TargetName);
+		PlatformsSettings->SaveConfig();
+	}
+
+	static bool CookBuildTargetIsChecked(FString TargetName)
+	{
+		return GetDefault<UPlatformsMenuSettings>()->CookBuildTarget == TargetName;
+	}
+
+	static bool CanCookContent()
+	{
+		return CanCookOrPackage(GetDefault<UPlatformsMenuSettings>()->CookPlatform, EPrepareContentMode::CookOnly);
+	}
+
+	static void CookContent()
+	{
+		const UPlatformsMenuSettings* PlatformsSettings = GetDefault<UPlatformsMenuSettings>();
+		CookOrPackage(PlatformsSettings->CookPlatform, PlatformsSettings->CookBuildTarget, EPrepareContentMode::CookOnly);
+	}
+
+	static void SetPackagePlatform(FName IniPlatformName)
+	{
+		UPlatformsMenuSettings* PlatformsSettings = GetMutableDefault<UPlatformsMenuSettings>();
+		PlatformsSettings->PackagePlatform = IniPlatformName;
+		PlatformsSettings->SaveConfig();
+	}
+
+	static bool PackagePlatformIsChecked(FName IniPlatformName)
+	{
+		return GetDefault<UPlatformsMenuSettings>()->PackagePlatform == IniPlatformName;
+	}
+
+	static bool CanPackageProject()
+	{
+		return CanCookOrPackage(GetDefault<UPlatformsMenuSettings>()->PackagePlatform, EPrepareContentMode::Package);
+	}
+
+	static void PackageProject()
+	{
+		const UPlatformsMenuSettings* PlatformsSettings = GetDefault<UPlatformsMenuSettings>();
+		CookOrPackage(PlatformsSettings->PackagePlatform, PlatformsSettings->PackageBuildTarget, EPrepareContentMode::Package);
+	}
+
+	static void SetPrepareForDebuggingPlatform(FName IniPlatformName)
+	{
+		UPlatformsMenuSettings* PlatformsSettings = GetMutableDefault<UPlatformsMenuSettings>();
+		PlatformsSettings->PrepareForDebuggingPlatform = IniPlatformName;
+		PlatformsSettings->SaveConfig();
+	}
+
+	static bool PrepareForDebuggingPlatformIsChecked(FName IniPlatformName)
+	{
+		return GetDefault<UPlatformsMenuSettings>()->PrepareForDebuggingPlatform == IniPlatformName;
+	}
+
+	static void SetPrepareForDebuggingBuildTarget(FString TargetName)
+	{
+		UPlatformsMenuSettings* PlatformsSettings = GetMutableDefault<UPlatformsMenuSettings>();
+		PlatformsSettings->PrepareForDebuggingBuildTarget = MoveTemp(TargetName);
+		PlatformsSettings->SaveConfig();
+	}
+
+	static bool PrepareForDebuggingBuildTargetIsChecked(FString TargetName)
+	{
+		return GetDefault<UPlatformsMenuSettings>()->PrepareForDebuggingBuildTarget == TargetName;
+	}
+
+	static bool CanPrepareForDebugging()
+	{
+		return CanCookOrPackage(GetDefault<UPlatformsMenuSettings>()->PrepareForDebuggingPlatform, EPrepareContentMode::PrepareForDebugging);
+	}
+
+	static void PrepareForDebugging()
+	{
+		const UPlatformsMenuSettings* PlatformsSettings = GetDefault<UPlatformsMenuSettings>();
+		CookOrPackage(PlatformsSettings->PrepareForDebuggingPlatform, PlatformsSettings->PrepareForDebuggingBuildTarget, EPrepareContentMode::PrepareForDebugging);
 	}
 
 	static void SetRefreshPlatformStatus()
@@ -1174,42 +1348,6 @@ static bool HasBuildForDeviceOrNot(const TArray<FProjectBuildSettings> Builds, b
 	return true;
 }
 
-static void MakeCustomBuildMenuEntries(FToolMenuSection& Section, const TArray<FProjectBuildSettings> Builds, FName IniPlatformName, const FString& DeviceId, const FString& DeviceName, const FText& ToolTipFormat)
-{
-	FString PlatformString = IniPlatformName.ToString();
-
-	for (FProjectBuildSettings Build : Builds)
-	{
-		if (Build.SpecificPlatforms.Num() == 0 || Build.SpecificPlatforms.Contains(PlatformString))
-		{
-			const bool bHasDeviceEntry = Build.BuildCookRunParams.Contains(TEXT("{DeviceId}"));
-			const bool bNeedsDeviceEntry = !DeviceId.IsEmpty();
-			if (bHasDeviceEntry == bNeedsDeviceEntry)
-			{
-				FString CommandLine = FTurnkeySupportCallbacks::GetCustomBuildCommandLine(IniPlatformName, DeviceId, Build);
-				FFormatNamedArguments Args;
-				Args.Add(TEXT("CommandLine"), FText::FromString(CommandLine));
-				Args.Add(TEXT("BuildName"), FText::FromString(Build.Name));
-				Args.Add(TEXT("BuildHelp"), FText::FromString(Build.HelpText));
-				Args.Add(TEXT("DeviceName"), FText::FromString(DeviceName));
-
-				Section.AddMenuEntry(
-					NAME_None,
-					FText::FromString(Build.Name),
-					FText::Format(ToolTipFormat, Args),
-					FSlateIcon(),
-					FUIAction(
-						FExecuteAction::CreateStatic(&FTurnkeySupportCallbacks::ExecuteCustomBuild, IniPlatformName, DeviceId, Build),
-						FCanExecuteAction::CreateStatic(&FTurnkeySupportCallbacks::CanExecuteCustomBuild, IniPlatformName, Build)
-					)
-				);
-			}
-		}
-	}
-
-}
-
-
 static void TurnkeySetDeviceAutoSoftwareUpdates(FString PlatformName, FString DeviceId, bool bEnableAutoSoftwareUpdates)
 {
 	FString CommandLine;
@@ -1228,129 +1366,420 @@ static void TurnkeySetDeviceAutoSoftwareUpdates(FString PlatformName, FString De
 		});	
 }
 
-
-static void MakeTurnkeyPlatformMenu(UToolMenu* ToolMenu, FName IniPlatformName, ITargetDeviceServicesModule* TargetDeviceServicesModule)
+// Returns the IniPlatformName of all platforms categorized into "supported", "unsupported", and "uncompiled" matching
+// the different behaviors used by the Platforms menu. "Uncompiled" indicates there is no compiled-in support for that
+// platform. "Unsupported" indicates the platform was disabled via the project settings.
+static void GetTurnkeyCategorizedPlatforms(TArray<FName>* OutSupportedPlatforms, TArray<FName>* OutUnsupportedPlatforms, TArray<FName>* OutUncompiledPlatforms)
 {
-	const FDataDrivenPlatformInfo& DDPI = FDataDrivenPlatformInfoRegistry::GetPlatformInfo(IniPlatformName);
-	FString UBTPlatformString = DDPI.UBTPlatformString;
+	if (OutSupportedPlatforms != nullptr)
+	{
+		OutSupportedPlatforms->Empty();
+	}
 
-	const PlatformInfo::FTargetPlatformInfo* VanillaInfo = PlatformInfo::FindVanillaPlatformInfo(IniPlatformName);
+	if (OutUnsupportedPlatforms != nullptr)
+	{
+		OutUnsupportedPlatforms->Empty();
+	}
 
-	const ITargetPlatform* TargetPlatform = GetTargetPlatformManager()->FindTargetPlatform(IniPlatformName);
+	if (OutUncompiledPlatforms != nullptr)
+	{
+		OutUncompiledPlatforms->Empty();
+	}
 
-	if (VanillaInfo != nullptr)
+	FProjectStatus ProjectStatus;
+	bool bProjectStatusIsValid = IProjectManager::Get().QueryStatusForCurrentProject(ProjectStatus);
+
+	for (const auto& Pair : FDataDrivenPlatformInfoRegistry::GetAllPlatformInfos())
 	{
-		FToolMenuSection& Section = ToolMenu->AddSection("ContentManagement", LOCTEXT("TurnkeySection_Content", "Content Management"));
+		if (Pair.Value.bIsFakePlatform || Pair.Value.bEnabledForUse == false)
+		{
+			continue;
+		}
 
-		Section.AddMenuEntry(
-			NAME_None,
-			LOCTEXT("Turnkey_PackageProject", "Package Project"),
-			LOCTEXT("TurnkeyTooltip_PackageProject", "Package this project and archive it to a user-selected directory. This can then be used to install and run."),
-			FSlateIcon(),
-			FUIAction(
-				FExecuteAction::CreateStatic(&FTurnkeySupportCallbacks::CookOrPackage, IniPlatformName, EPrepareContentMode::Package),
-				FCanExecuteAction::CreateStatic(&FTurnkeySupportCallbacks::CanCookOrPackage, IniPlatformName, EPrepareContentMode::Package)
-			)
-		);
+		FName IniPlatformName = Pair.Key;
+		if (FDataDrivenPlatformInfoRegistry::IsPlatformHiddenFromUI(IniPlatformName))
+		{
+			continue;
+		}
+
+		if (!FDataDrivenPlatformInfoRegistry::HasCompiledSupportForPlatform(
+				IniPlatformName, FDataDrivenPlatformInfoRegistry::EPlatformNameType::Ini
+			))
+		{
+			if (OutUncompiledPlatforms != nullptr)
+			{
+				OutUncompiledPlatforms->Add(IniPlatformName);
+			}
+
+			continue;
+		}
+
+		if (bProjectStatusIsValid && !ProjectStatus.IsTargetPlatformSupported(IniPlatformName))
+		{
+			if (OutUnsupportedPlatforms != nullptr)
+			{
+				OutUnsupportedPlatforms->Add(IniPlatformName);
+			}
+
+			continue;
+		}
+
+		if (OutSupportedPlatforms != nullptr)
+		{
+			OutSupportedPlatforms->Add(IniPlatformName);
+		}
+	}
+}
+
+// Shows a sub-menu for the platform allowing selection of the flavor or architecture where appropriate
+static void MakeTurnkeyPerPlatformMenu(UToolMenu* Menu, FName IniPlatformName, TArray<const PlatformInfo::FTargetPlatformInfo*> ValidFlavors, TArray<FString> Architectures)
+{
+	if (ValidFlavors.Num() > 1)
+	{
+		FToolMenuSection& FlavorSection = Menu->AddSection("FlavorSelection", LOCTEXT("TurnkeySection_FlavorSelection", "Flavor Selection"));
+					
+		for (const PlatformInfo::FTargetPlatformInfo* Info : ValidFlavors)
+		{
+			FlavorSection.AddMenuEntry(
+				NAME_None,
+				Info->DisplayName,
+				FText(),
+				FSlateIcon(),
+				FUIAction(
+					FExecuteAction::CreateStatic(&FTurnkeySupportCallbacks::SetActiveFlavor, Info),
+					FCanExecuteAction(),
+					FIsActionChecked::CreateStatic(&FTurnkeySupportCallbacks::SetActiveFlavorIsChecked, Info)
+				),
+				EUserInterfaceActionType::RadioButton
+			);
+		}
+	}
+
+	if (Architectures.Num() > 1)
+	{
+		FToolMenuSection& ArchitectureSection = Menu->AddSection("ArchitectureSelection", LOCTEXT("TurnkeySection_ArchitectureSelection", "Architecture Selection"));
 
-		Section.AddMenuEntry(
+		ArchitectureSection.AddMenuEntry(
 			NAME_None,
-			LOCTEXT("Turnkey_CookContent", "Cook Content"),
-			LOCTEXT("TurnkeyTooltip_CookContent", "Cook this project for the selected configuration and target"),
+			LOCTEXT("DefaultArchitectureName", "Project Default"),
+			FText(),
 			FSlateIcon(),
 			FUIAction(
-				FExecuteAction::CreateStatic(&FTurnkeySupportCallbacks::CookOrPackage, IniPlatformName, EPrepareContentMode::CookOnly),
-				FCanExecuteAction::CreateStatic(&FTurnkeySupportCallbacks::CanCookOrPackage, IniPlatformName, EPrepareContentMode::CookOnly)
-			)
+				FExecuteAction::CreateStatic(&FTurnkeySupportCallbacks::SetActiveArchitecture, IniPlatformName, FString()),
+				FCanExecuteAction(),
+				FIsActionChecked::CreateStatic(&FTurnkeySupportCallbacks::SetActiveArchitectureIsChecked, IniPlatformName, FString())
+			),
+			EUserInterfaceActionType::RadioButton
 		);
-        // add platforms here when they have their own Prepare for debug flow
-#if PLATFORM_WINDOWS || PLATFORM_MAC
-		if (IniPlatformName.ToString() == "IOS" || IniPlatformName.ToString() == "TvOS")
+
+		const ITargetPlatform* TargetPlatform = GetTargetPlatformManager()->FindTargetPlatform(IniPlatformName);
+
+		for (const FString& Architecture : Architectures)
 		{
+			bool bIsHostArchitecture = (IniPlatformName == FPlatformProperties::IniPlatformName()) && (Architecture == TargetPlatform->GetHostArchitecture());
 
-			FString Tooltip = "";
-			
-			Section.AddMenuEntry(
+			FText ArchLabel = bIsHostArchitecture 
+				? FText::Format(LOCTEXT("HostArchLabel", "{0} (this platform)"), FText::FromString(Architecture)) 
+				: FText::FromString(Architecture);
+
+			ArchitectureSection.AddMenuEntry(
 				NAME_None,
-				LOCTEXT("Turnkey_PrepareForDebugging", "Prepare For Debugging"),
-				FText::Format(LOCTEXT("TurnkeyTooltip_PrepareForDebugging", "Expects a working remote toolchain and an IPA package with the same name as the project file in the Binaries/{0}/ folder."), FText::FromString(IniPlatformName.ToString().ToUpper())),
+				ArchLabel,
+				FText(),
 				FSlateIcon(),
 				FUIAction(
-					FExecuteAction::CreateStatic(&FTurnkeySupportCallbacks::CookOrPackage, IniPlatformName, EPrepareContentMode::PrepareForDebugging),
-					FCanExecuteAction::CreateStatic(&FTurnkeySupportCallbacks::CanCookOrPackage, IniPlatformName, EPrepareContentMode::PrepareForDebugging)
-				)
+					FExecuteAction::CreateStatic(&FTurnkeySupportCallbacks::SetActiveArchitecture, IniPlatformName, Architecture),
+					FCanExecuteAction(),
+					FIsActionChecked::CreateStatic(&FTurnkeySupportCallbacks::SetActiveArchitectureIsChecked, IniPlatformName, Architecture)
+				),
+				EUserInterfaceActionType::RadioButton
 			);
 		}
-#endif
+	}
+}
 
-		UProjectPackagingSettings* PackagingSettings = FTurnkeySupportCallbacks::GetPackagingSettingsForPlatform(IniPlatformName);
+static void MakeCookContentMenuItems(UToolMenu* Menu)
+{
+	FToolMenuSection& TopSection = Menu->AddSection("CookSection");
+	{
+		TopSection.AddMenuEntry(
+			NAME_None,
+			LOCTEXT("Turnkey_CookContent", "Cook Content"),
+			LOCTEXT("TurnkeyTooltip_CookContent", "Cook this project for the selected Platform and target"),
+			FSlateIcon(FAppStyle::GetAppStyleSetName(), "MainFrame.CookContent"),
+			FUIAction(
+				FExecuteAction::CreateStatic(&FTurnkeySupportCallbacks::CookContent),
+				FCanExecuteAction::CreateStatic(&FTurnkeySupportCallbacks::CanCookContent)
+			)
+		);
+	}
 
+	FToolMenuSection& PlatformsSection = Menu->AddSection("CookPlatformsSection", LOCTEXT("TurnkeySection_CookPlatforms", "Cook For"));
+	{
+		PlatformsSection.AddDynamicEntry(NAME_None, FNewToolMenuSectionDelegate::CreateLambda([](FToolMenuSection& PlatformsSection)
+		{
+			ITargetDeviceServicesModule* TargetDeviceServicesModule = static_cast<ITargetDeviceServicesModule*>(FModuleManager::Get().LoadModule(TEXT("TargetDeviceServices")));
 
-			
-		// these ToolTipFormats will be formatted inside the MakeCustomBuildMenuEntries function with some named args
-		// the empty strings passed in is the deviceId, which means to show builds that don't contain DeviceId components
-#if ALLOW_CONTROL_TO_COPY_COMMANDLINE
-		FText ToolTipFormat = LOCTEXT("TurnkeyTooltip_EngineCustomBuild_WithCopy", "Execute the '{BuildName}' command.\n{BuildHelp}\n--\nThis runs the following command:\nRunUAT {CommandLine}\nHold Control to copy the final underlying BuildCookRun commandline to the clipboard.");
-#else
-		FText ToolTipFormat = LOCTEXT("TurnkeyTooltip_EngineCustomBuild", "Execute the '{BuildName}' command.\n{BuildHelp}\n--\nThis runs the following command:\nRunUAT {CommandLine}");
-#endif
-		MakeCustomBuildMenuEntries(Section, PackagingSettings->EngineCustomBuilds, IniPlatformName, FString(), FString(), ToolTipFormat);
-		
-#if ALLOW_CONTROL_TO_COPY_COMMANDLINE
-		ToolTipFormat = LOCTEXT("TurnkeyTooltip_ProjectCustomBuild_WithCopy", "Execute this project's custom '{BuildName}' command.\n{BuildHelp}\n--\nThis runs the following command:\nRunUAT {CommandLine}\nHold Control to copy the final BuildCookRun commandline to the clipboard.\nThis custom command comes from Project Settings.");
-#else
-		ToolTipFormat = LOCTEXT("TurnkeyTooltip_ProjectCustomBuild", "Execute this project's custom '{BuildName}' command.\n{BuildHelp}\n--\nThis runs the following command:\nRunUAT {CommandLine}\nThis custom command comes from Project Settings.");
-#endif
-		MakeCustomBuildMenuEntries(Section, PackagingSettings->ProjectCustomBuilds, IniPlatformName, FString(), FString(), ToolTipFormat);
+			TArray<FName> SupportedPlatforms;
+			GetTurnkeyCategorizedPlatforms(&SupportedPlatforms, nullptr, nullptr);
 
+			for (const FName& IniPlatformName : SupportedPlatforms)
+			{
+				const PlatformInfo::FTargetPlatformInfo* VanillaInfo = PlatformInfo::FindVanillaPlatformInfo(IniPlatformName);
 
-		UProjectPackagingSettings* AllPlatformPackagingSettings = GetMutableDefault<UProjectPackagingSettings>();
+				const ITargetPlatform* TargetPlatform = GetTargetPlatformManager()->FindTargetPlatform(IniPlatformName);
 
-		// Populate Flavor Selection menu with available flavors
-		// Flavor Selection exists for Android to be able to package ASTC, DXT, ETC2, etc
-		{
-			// gather all valid flavors
-			const TArray<const PlatformInfo::FTargetPlatformInfo*> ValidFlavors = VanillaInfo->Flavors.FilterByPredicate([](const PlatformInfo::FTargetPlatformInfo* Target)
+				// Flavor selection exists for Android to be able to package ASTC, DXT, ETC2, etc
+				TArray<const PlatformInfo::FTargetPlatformInfo*> ValidFlavors;
+
+				// Architecture selection exists for platforms targeting multiple CPU architectures, such as Windows on
+				// x64 and ARM64.
+				TArray<FString> Architectures;
+
+				if (VanillaInfo != nullptr)
 				{
-					// Editor isn't a valid platform type that users can target
-					// The Build Target will choose client or server, so no need to show them as well
-					return Target->PlatformType != EBuildTargetType::Editor && Target->PlatformType != EBuildTargetType::Client && Target->PlatformType != EBuildTargetType::Server;
-				});
+					ValidFlavors = VanillaInfo->Flavors.FilterByPredicate([](const PlatformInfo::FTargetPlatformInfo* Target)
+						{
+							// Editor isn't a valid platform type that users can target
+							// The Build Target will choose client or server, so no need to show them as well
+							return Target->PlatformType != EBuildTargetType::Editor && Target->PlatformType != EBuildTargetType::Client && Target->PlatformType != EBuildTargetType::Server;
+						});
+				}
 
-			if (ValidFlavors.Num() > 1)
-			{
-				FToolMenuSection& FlavorSection = ToolMenu->AddSection("FlavorSelection", LOCTEXT("TurnkeySection_FlavorSelection", "Flavor Selection"));
-				
-				for (const PlatformInfo::FTargetPlatformInfo* Info : ValidFlavors)
+				if (TargetPlatform != nullptr)
+				{
+					TargetPlatform->GetPossibleArchitectures(Architectures);
+				}
+			
+				// Platforms that have multiple flavors or architectures are displayed as a radio button that also pops
+				// out a sub-menu to choose the flavor and/or architecture.
+				if (ValidFlavors.Num() > 1 || Architectures.Num() > 1)
+				{
+					FToolMenuEntry Entry = FToolMenuEntry::InitSubMenu(
+						NAME_None,
+						MakeSdkStatusAttribute(IniPlatformName, nullptr),
+						FText::FromString(IniPlatformName.ToString()),
+						FNewToolMenuDelegate::CreateStatic(&MakeTurnkeyPerPlatformMenu, IniPlatformName, ValidFlavors, Architectures),
+						FUIAction(
+							FExecuteAction::CreateStatic(&FTurnkeySupportCallbacks::SetCookPlatform, IniPlatformName),
+							FCanExecuteAction(),
+							FIsActionChecked::CreateStatic(&FTurnkeySupportCallbacks::CookPlatformIsChecked, IniPlatformName)
+						),
+						EUserInterfaceActionType::RadioButton,
+						false,
+						MakePlatformSdkIconAttribute(IniPlatformName, nullptr)
+					);
+
+					PlatformsSection.AddEntry(Entry);
+				}
+				else
 				{
-					FlavorSection.AddMenuEntry(
+					FToolMenuEntry Entry = FToolMenuEntry::InitMenuEntry(
 						NAME_None,
-						Info->DisplayName,
-						FText(),
-						FSlateIcon(),
+						MakeSdkStatusAttribute(IniPlatformName, nullptr),
+						FText::FromString(IniPlatformName.ToString()),
+						MakePlatformSdkIconAttribute(IniPlatformName, nullptr),
 						FUIAction(
-							FExecuteAction::CreateStatic(&FTurnkeySupportCallbacks::SetActiveFlavor, Info),
-							FCanExecuteAction::CreateStatic(&FTurnkeySupportCallbacks::CanSetActiveFlavor, Info),
-							FIsActionChecked::CreateStatic(&FTurnkeySupportCallbacks::SetActiveFlavorIsChecked, Info)
+							FExecuteAction::CreateStatic(&FTurnkeySupportCallbacks::SetCookPlatform, IniPlatformName),
+							FCanExecuteAction(),
+							FIsActionChecked::CreateStatic(&FTurnkeySupportCallbacks::CookPlatformIsChecked, IniPlatformName)
 						),
 						EUserInterfaceActionType::RadioButton
 					);
+
+					PlatformsSection.AddEntry(Entry);
 				}
 			}
-		}
+		}));
+	}
 
-		FToolMenuSection& ConfigSection = ToolMenu->AddSection("BuildConfig", LOCTEXT("TurnkeySection_BuildConfig", "Binary Configuration"));
-		
-		const UProjectPackagingSettings::FConfigurationInfo& DefaultConfigInfo = UProjectPackagingSettings::ConfigurationInfo[static_cast<int32>(PackagingSettings->BuildConfiguration)];
-		ConfigSection.AddMenuEntry(
+	UProjectPackagingSettings* AllPlatformPackagingSettings = GetMutableDefault<UProjectPackagingSettings>();
+
+	// Collect build targets. Content-only projects use Engine targets
+	FProjectStatus ProjectStatus;
+	bool bHasCode = IProjectManager::Get().QueryStatusForCurrentProject(ProjectStatus) && ProjectStatus.bCodeBasedProject;
+
+	IDesktopPlatform* DesktopPlatform = FDesktopPlatformModule::Get();
+	TArray<FTargetInfo> Targets = bHasCode ? DesktopPlatform->GetTargetsForCurrentProject() : DesktopPlatform->GetTargetsForProject(FString());
+
+	if (Targets.Num() > 0)
+	{
+		Targets.Sort([](const FTargetInfo& A, const FTargetInfo& B) { return A.Name < B.Name; });
+		
+		const TArray<FTargetInfo> ValidTargets = Targets.FilterByPredicate([](const FTargetInfo& Target)
+			{
+				return Target.Type == EBuildTargetType::Game || Target.Type == EBuildTargetType::Client || Target.Type == EBuildTargetType::Server;
+			});
+
+		if (ValidTargets.Num() > 1)
+		{
+			// Set BuildTarget to default to Game if it hasn't been set (or if it was set, but is no longer valid)
+			FString BuildTarget = AllPlatformPackagingSettings->BuildTarget;
+			if (BuildTarget.IsEmpty() ||
+			   !ValidTargets.ContainsByPredicate([&BuildTarget](const FTargetInfo& Target) { return Target.Name == BuildTarget; } ))
+			{
+				TArray<FTargetInfo> GameTargets = ValidTargets.FilterByPredicate([](const FTargetInfo& Target)
+					{ return Target.Type == EBuildTargetType::Game;	});
+				// if there are no Game targets, look for a Client target
+				if (GameTargets.Num() == 0)
+				{
+					GameTargets = ValidTargets.FilterByPredicate([](const FTargetInfo& Target)
+						{ return Target.Type == EBuildTargetType::Client;	});
+				}
+				// if no clients, then just _anything_ 
+				if (GameTargets.Num() == 0)
+				{
+					GameTargets = ValidTargets;
+				}
+
+				AllPlatformPackagingSettings->BuildTarget = GameTargets[0].Name;
+				AllPlatformPackagingSettings->SaveConfig();
+			}
+
+			FToolMenuSection& TargetSection = Menu->AddSection("BuildTarget", LOCTEXT("TurnkeySection_BuildTarget", "Build Target"));
+
+			// Using the project default setting is represented by an empty string
+			TargetSection.AddMenuEntry(
+				NAME_None,
+				FText::Format(LOCTEXT("DefaultPackageTarget",  "Use Project Setting ({0})"), FText::FromString(AllPlatformPackagingSettings->BuildTarget)),
+				FText::Format(LOCTEXT("DefaultPackageTargetTooltip", "Package the {0} target"), FText::FromString(AllPlatformPackagingSettings->BuildTarget)),
+				FSlateIcon(),
+				FUIAction(
+					FExecuteAction::CreateStatic(&FTurnkeySupportCallbacks::SetCookBuildTarget, FString()),
+					FCanExecuteAction(),
+					FIsActionChecked::CreateStatic(&FTurnkeySupportCallbacks::CookBuildTargetIsChecked, FString())
+					),
+				EUserInterfaceActionType::RadioButton
+			);
+
+			for (const FTargetInfo& Target : ValidTargets)
+			{
+				TargetSection.AddMenuEntry(
+					NAME_None,
+					FText::FromString(Target.Name),
+					FText::Format(LOCTEXT("PackageTargetName", "Package the '{0}' target."), FText::FromString(Target.Name)),
+					FSlateIcon(),
+					FUIAction(
+						FExecuteAction::CreateStatic(&FTurnkeySupportCallbacks::SetCookBuildTarget, Target.Name),
+						FCanExecuteAction(),
+						FIsActionChecked::CreateStatic(&FTurnkeySupportCallbacks::CookBuildTargetIsChecked, Target.Name)
+					),
+					EUserInterfaceActionType::RadioButton
+				);
+			}
+		}
+	}
+}
+
+static void MakePackageProjectMenuItems(UToolMenu* Menu)
+{
+	FToolMenuSection& TopSection = Menu->AddSection("PackageSection");
+	{
+		TopSection.AddMenuEntry(
+			"PackageProject",
+			LOCTEXT("Turnkey_PackageProject", "Package Project"),
+			LOCTEXT("TurnkeyTooltip_PackageProject", "Package this project and archive it to a user-selected directory. This can then be used to install and run."),
+			FSlateIcon(FAppStyle::GetAppStyleSetName(), "MainFrame.PackageProject"),
+			FUIAction(
+				FExecuteAction::CreateStatic(&FTurnkeySupportCallbacks::PackageProject),
+				FCanExecuteAction::CreateStatic(&FTurnkeySupportCallbacks::CanPackageProject)
+			)
+		);
+	}
+
+	FToolMenuSection& PlatformsSection = Menu->AddSection("PackagePlatformsSection", LOCTEXT("PackagePlatformsSection", "Package For"));
+	{
+		PlatformsSection.AddDynamicEntry(NAME_None, FNewToolMenuSectionDelegate::CreateLambda([](FToolMenuSection& PlatformsSection)
+		{
+			ITargetDeviceServicesModule* TargetDeviceServicesModule = static_cast<ITargetDeviceServicesModule*>(FModuleManager::Get().LoadModule(TEXT("TargetDeviceServices")));
+
+			TArray<FName> SupportedPlatforms;
+			GetTurnkeyCategorizedPlatforms(&SupportedPlatforms, nullptr, nullptr);
+
+			for (const FName& IniPlatformName : SupportedPlatforms)
+			{
+				const PlatformInfo::FTargetPlatformInfo* VanillaInfo = PlatformInfo::FindVanillaPlatformInfo(IniPlatformName);
+
+				const ITargetPlatform* TargetPlatform = GetTargetPlatformManager()->FindTargetPlatform(IniPlatformName);
+
+				// Flavor selection exists for Android to be able to package ASTC, DXT, ETC2, etc
+				TArray<const PlatformInfo::FTargetPlatformInfo*> ValidFlavors;
+
+				// Architecture selection exists for platforms targeting multiple CPU architectures, such as Windows on
+				// x64 and ARM64.
+				TArray<FString> Architectures;
+
+				if (VanillaInfo != nullptr)
+				{
+					ValidFlavors = VanillaInfo->Flavors.FilterByPredicate([](const PlatformInfo::FTargetPlatformInfo* Target)
+						{
+							// Editor isn't a valid platform type that users can target
+							// The Build Target will choose client or server, so no need to show them as well
+							return Target->PlatformType != EBuildTargetType::Editor && Target->PlatformType != EBuildTargetType::Client && Target->PlatformType != EBuildTargetType::Server;
+						});
+				}
+
+				if (TargetPlatform != nullptr)
+				{
+					TargetPlatform->GetPossibleArchitectures(Architectures);
+				}
+			
+				// Platforms that have multiple flavors or architectures are displayed as a radio button that also pops
+				// out a sub-menu to choose the flavor and/or architecture.
+				if (ValidFlavors.Num() > 1 || Architectures.Num() > 1)
+				{
+					FToolMenuEntry Entry = FToolMenuEntry::InitSubMenu(
+						NAME_None,
+						MakeSdkStatusAttribute(IniPlatformName, nullptr),
+						FText::FromString(IniPlatformName.ToString()),
+						FNewToolMenuDelegate::CreateStatic(&MakeTurnkeyPerPlatformMenu, IniPlatformName, ValidFlavors, Architectures),
+						FUIAction(
+							FExecuteAction::CreateStatic(&FTurnkeySupportCallbacks::SetPackagePlatform, IniPlatformName),
+							FCanExecuteAction(),
+							FIsActionChecked::CreateStatic(&FTurnkeySupportCallbacks::PackagePlatformIsChecked, IniPlatformName)
+						),
+						EUserInterfaceActionType::RadioButton,
+						false,
+						MakePlatformSdkIconAttribute(IniPlatformName, nullptr)
+					);
+
+					PlatformsSection.AddEntry(Entry);
+				}
+				else
+				{
+					FToolMenuEntry Entry = FToolMenuEntry::InitMenuEntry(
+						NAME_None,
+						MakeSdkStatusAttribute(IniPlatformName, nullptr),
+						FText::FromString(IniPlatformName.ToString()),
+						MakePlatformSdkIconAttribute(IniPlatformName, nullptr),
+						FUIAction(
+							FExecuteAction::CreateStatic(&FTurnkeySupportCallbacks::SetPackagePlatform, IniPlatformName),
+							FCanExecuteAction(),
+							FIsActionChecked::CreateStatic(&FTurnkeySupportCallbacks::PackagePlatformIsChecked, IniPlatformName)
+						),
+						EUserInterfaceActionType::RadioButton
+					);
+
+					PlatformsSection.AddEntry(Entry);
+				}
+			}
+		}));
+	}
+
+	UProjectPackagingSettings* AllPlatformPackagingSettings = GetMutableDefault<UProjectPackagingSettings>();
+
+	FToolMenuSection& ConfigSection = Menu->AddSection("BuildConfig", LOCTEXT("TurnkeySection_BuildConfig", "Build Configuration"));
+	{
+		const UProjectPackagingSettings::FConfigurationInfo& DefaultConfigInfo = UProjectPackagingSettings::ConfigurationInfo[static_cast<int32>(AllPlatformPackagingSettings->BuildConfiguration)];
+		ConfigSection.AddMenuEntry(
 			NAME_None,
 			FText::Format(LOCTEXT("DefaultConfiguration",  "Use Project Setting ({0})"), DefaultConfigInfo.Name),
 			DefaultConfigInfo.ToolTip,
 			FSlateIcon(),
 			FUIAction(
-				FExecuteAction::CreateStatic(&FTurnkeySupportCallbacks::SetPackageBuildConfiguration, VanillaInfo, EProjectPackagingBuildConfigurations::PPBC_MAX),
+				FExecuteAction::CreateStatic(&FTurnkeySupportCallbacks::SetPackageBuildConfiguration, EProjectPackagingBuildConfigurations::PPBC_MAX),
 				FCanExecuteAction(),
-				FIsActionChecked::CreateStatic(&FTurnkeySupportCallbacks::PackageBuildConfigurationIsChecked, VanillaInfo, EProjectPackagingBuildConfigurations::PPBC_MAX)
+				FIsActionChecked::CreateStatic(&FTurnkeySupportCallbacks::PackageBuildConfigurationIsChecked, EProjectPackagingBuildConfigurations::PPBC_MAX)
 			),
 			EUserInterfaceActionType::RadioButton
 		);
@@ -1369,136 +1798,549 @@ static void MakeTurnkeyPlatformMenu(UToolMenu* ToolMenu, FName IniPlatformName,
 					ConfigurationInfo.ToolTip,
 					FSlateIcon(),
 					FUIAction(
-						FExecuteAction::CreateStatic(&FTurnkeySupportCallbacks::SetPackageBuildConfiguration, VanillaInfo, PackagingConfiguration),
+						FExecuteAction::CreateStatic(&FTurnkeySupportCallbacks::SetPackageBuildConfiguration, PackagingConfiguration),
 						FCanExecuteAction(),
-						FIsActionChecked::CreateStatic(&FTurnkeySupportCallbacks::PackageBuildConfigurationIsChecked, VanillaInfo, PackagingConfiguration)
+						FIsActionChecked::CreateStatic(&FTurnkeySupportCallbacks::PackageBuildConfigurationIsChecked, PackagingConfiguration)
 					),
 					EUserInterfaceActionType::RadioButton
 				);
 			}
 		}
+	}
 
-		// Collect build targets. Content-only projects use Engine targets
-		FProjectStatus ProjectStatus;
-		bool bHasCode = IProjectManager::Get().QueryStatusForCurrentProject(ProjectStatus) && ProjectStatus.bCodeBasedProject;
+	// Collect build targets. Content-only projects use Engine targets
+	FProjectStatus ProjectStatus;
+	bool bHasCode = IProjectManager::Get().QueryStatusForCurrentProject(ProjectStatus) && ProjectStatus.bCodeBasedProject;
 
-		IDesktopPlatform* DesktopPlatform = FDesktopPlatformModule::Get();
-		TArray<FTargetInfo> Targets = bHasCode ? DesktopPlatform->GetTargetsForCurrentProject() : DesktopPlatform->GetTargetsForProject(FString());
+	IDesktopPlatform* DesktopPlatform = FDesktopPlatformModule::Get();
+	TArray<FTargetInfo> Targets = bHasCode ? DesktopPlatform->GetTargetsForCurrentProject() : DesktopPlatform->GetTargetsForProject(FString());
 
-		if (Targets.Num() > 0)
+	if (Targets.Num() > 0)
+	{
+		Targets.Sort([](const FTargetInfo& A, const FTargetInfo& B) { return A.Name < B.Name; });
+		
+		const TArray<FTargetInfo> ValidTargets = Targets.FilterByPredicate([](const FTargetInfo& Target)
+			{
+				return Target.Type == EBuildTargetType::Game || Target.Type == EBuildTargetType::Client || Target.Type == EBuildTargetType::Server;
+			});
+
+		if (ValidTargets.Num() > 1)
 		{
-			Targets.Sort([](const FTargetInfo& A, const FTargetInfo& B) { return A.Name < B.Name; });
-			
-			const TArray<FTargetInfo> ValidTargets = Targets.FilterByPredicate([](const FTargetInfo& Target)
+			// Set BuildTarget to default to Game if it hasn't been set (or if it was set, but is no longer valid)
+			FString BuildTarget = AllPlatformPackagingSettings->BuildTarget;
+			if (BuildTarget.IsEmpty() ||
+			   !ValidTargets.ContainsByPredicate([&BuildTarget](const FTargetInfo& Target) { return Target.Name == BuildTarget; } ))
+			{
+				TArray<FTargetInfo> GameTargets = ValidTargets.FilterByPredicate([](const FTargetInfo& Target)
+					{ return Target.Type == EBuildTargetType::Game;	});
+				// if there are no Game targets, look for a Client target
+				if (GameTargets.Num() == 0)
 				{
-					return Target.Type == EBuildTargetType::Game || Target.Type == EBuildTargetType::Client || Target.Type == EBuildTargetType::Server;
-				});
+					GameTargets = ValidTargets.FilterByPredicate([](const FTargetInfo& Target)
+						{ return Target.Type == EBuildTargetType::Client;	});
+				}
+				// if no clients, then just _anything_ 
+				if (GameTargets.Num() == 0)
+				{
+					GameTargets = ValidTargets;
+				}
+
+				AllPlatformPackagingSettings->BuildTarget = GameTargets[0].Name;
+				AllPlatformPackagingSettings->SaveConfig();
+			}
+
+			FToolMenuSection& TargetSection = Menu->AddSection("BuildTarget", LOCTEXT("TurnkeySection_BuildTarget", "Build Target"));
+
+			// Using the project default setting is represented by an empty string
+			TargetSection.AddMenuEntry(
+				NAME_None,
+				FText::Format(LOCTEXT("DefaultPackageTarget",  "Use Project Setting ({0})"), FText::FromString(AllPlatformPackagingSettings->BuildTarget)),
+				FText::Format(LOCTEXT("DefaultPackageTargetTooltip", "Package the {0} target"), FText::FromString(AllPlatformPackagingSettings->BuildTarget)),
+				FSlateIcon(),
+				FUIAction(
+					FExecuteAction::CreateStatic(&FTurnkeySupportCallbacks::SetPackageBuildTarget, FString()),
+					FCanExecuteAction(),
+					FIsActionChecked::CreateStatic(&FTurnkeySupportCallbacks::PackageBuildTargetIsChecked, FString())
+					),
+				EUserInterfaceActionType::RadioButton
+			);
+
+			for (const FTargetInfo& Target : ValidTargets)
+			{
+				TargetSection.AddMenuEntry(
+					NAME_None,
+					FText::FromString(Target.Name),
+					FText::Format(LOCTEXT("PackageTargetName", "Package the '{0}' target."), FText::FromString(Target.Name)),
+					FSlateIcon(),
+					FUIAction(
+						FExecuteAction::CreateStatic(&FTurnkeySupportCallbacks::SetPackageBuildTarget, Target.Name),
+						FCanExecuteAction(),
+						FIsActionChecked::CreateStatic(&FTurnkeySupportCallbacks::PackageBuildTargetIsChecked, Target.Name)
+					),
+					EUserInterfaceActionType::RadioButton
+				);
+			}
+		}
+	}
+}
+
+static void MakePrepareForDebuggingMenuItems(UToolMenu* Menu)
+{
+	FToolMenuSection& TopSection = Menu->AddSection("PrepareSection");
+	{
+		TopSection.AddMenuEntry(
+			"PrepareForDebugging",
+			LOCTEXT("Turnkey_PrepareForDebugging", "Prepare For Debugging"),
+			LOCTEXT("Turnkey_PrepareForDebugging_Tooltip", "Create files to help streaming the debugging pipeline on some platforms"),
+			FSlateIcon(FAppStyle::GetAppStyleSetName(), "Debug"),
+			FUIAction(
+				FExecuteAction::CreateStatic(&FTurnkeySupportCallbacks::PrepareForDebugging),
+				FCanExecuteAction::CreateStatic(&FTurnkeySupportCallbacks::CanPrepareForDebugging)
+			)
+		);
+	}
 
-			if (ValidTargets.Num() > 1)
+	FToolMenuSection& PlatformsSection = Menu->AddSection("PlatformsSection", LOCTEXT("TurnkeySection_PrepareForDebuggingPlatforms", "Prepare For"));
+	{
+		PlatformsSection.AddDynamicEntry(NAME_None, FNewToolMenuSectionDelegate::CreateLambda([](FToolMenuSection& PlatformsSection)
+		{
+			TArray<FName> SupportedPlatforms;
+			GetTurnkeyCategorizedPlatforms(&SupportedPlatforms, nullptr, nullptr);
+
+			for (const FName& IniPlatformName : SupportedPlatforms)
 			{
-				// Set BuildTarget to default to Game if it hasn't been set (or if it was set, but is no longer valid)
-				FString BuildTarget = AllPlatformPackagingSettings->BuildTarget;
-				if (BuildTarget.IsEmpty() ||
-				   !ValidTargets.ContainsByPredicate([&BuildTarget](const FTargetInfo& Target) { return Target.Name == BuildTarget; } ))
+				if (IniPlatformName != TEXT("IOS") && IniPlatformName != TEXT("TvOS"))
 				{
-					TArray<FTargetInfo> GameTargets = ValidTargets.FilterByPredicate([](const FTargetInfo& Target)
-						{ return Target.Type == EBuildTargetType::Game;	});
-					// if there are no Game targets, look for a Client target
-					if (GameTargets.Num() == 0)
-					{
-						GameTargets = ValidTargets.FilterByPredicate([](const FTargetInfo& Target)
-							{ return Target.Type == EBuildTargetType::Client;	});
-					}
-					// if no clients, then just _anything_ 
-					if (GameTargets.Num() == 0)
-					{
-						GameTargets = ValidTargets;
-					}
+					continue;
+				}
 
-					AllPlatformPackagingSettings->BuildTarget = GameTargets[0].Name;
-					AllPlatformPackagingSettings->SaveConfig();
+				// Create a radio button per-platform
+				FToolMenuEntry Entry = FToolMenuEntry::InitMenuEntry(
+					NAME_None,
+					MakeSdkStatusAttribute(IniPlatformName, nullptr),
+					FText::Format(LOCTEXT("TurnkeyTooltip_PrepareForDebugging", "Expects a working remote toolchain and an IPA package with the same name as the project file in the Binaries/{0}/ folder."), FText::FromString(IniPlatformName.ToString().ToUpper())),
+					MakePlatformSdkIconAttribute(IniPlatformName, nullptr),
+					FUIAction(
+						FExecuteAction::CreateStatic(&FTurnkeySupportCallbacks::SetPrepareForDebuggingPlatform, IniPlatformName),
+						FCanExecuteAction(),
+						FIsActionChecked::CreateStatic(&FTurnkeySupportCallbacks::PrepareForDebuggingPlatformIsChecked, IniPlatformName)
+					),
+					EUserInterfaceActionType::RadioButton
+				);
+
+				PlatformsSection.AddEntry(Entry);
+			}
+		}));
+	}
+
+	UProjectPackagingSettings* AllPlatformPackagingSettings = GetMutableDefault<UProjectPackagingSettings>();
+
+	// Collect build targets. Content-only projects use Engine targets
+	FProjectStatus ProjectStatus;
+	bool bHasCode = IProjectManager::Get().QueryStatusForCurrentProject(ProjectStatus) && ProjectStatus.bCodeBasedProject;
+
+	IDesktopPlatform* DesktopPlatform = FDesktopPlatformModule::Get();
+	TArray<FTargetInfo> Targets = bHasCode ? DesktopPlatform->GetTargetsForCurrentProject() : DesktopPlatform->GetTargetsForProject(FString());
+
+	if (Targets.Num() > 0)
+	{
+		Targets.Sort([](const FTargetInfo& A, const FTargetInfo& B) { return A.Name < B.Name; });
+		
+		const TArray<FTargetInfo> ValidTargets = Targets.FilterByPredicate([](const FTargetInfo& Target)
+			{
+				return Target.Type == EBuildTargetType::Game || Target.Type == EBuildTargetType::Client || Target.Type == EBuildTargetType::Server;
+			});
+
+		if (ValidTargets.Num() > 1)
+		{
+			// Set BuildTarget to default to Game if it hasn't been set (or if it was set, but is no longer valid)
+			FString BuildTarget = AllPlatformPackagingSettings->BuildTarget;
+			if (BuildTarget.IsEmpty() ||
+			   !ValidTargets.ContainsByPredicate([&BuildTarget](const FTargetInfo& Target) { return Target.Name == BuildTarget; } ))
+			{
+				TArray<FTargetInfo> GameTargets = ValidTargets.FilterByPredicate([](const FTargetInfo& Target)
+					{ return Target.Type == EBuildTargetType::Game;	});
+				// if there are no Game targets, look for a Client target
+				if (GameTargets.Num() == 0)
+				{
+					GameTargets = ValidTargets.FilterByPredicate([](const FTargetInfo& Target)
+						{ return Target.Type == EBuildTargetType::Client;	});
+				}
+				// if no clients, then just _anything_ 
+				if (GameTargets.Num() == 0)
+				{
+					GameTargets = ValidTargets;
 				}
 
-				FToolMenuSection& TargetSection = ToolMenu->AddSection("BuildTarget", LOCTEXT("TurnkeySection_BuildTarget", "Build Target"));
+				AllPlatformPackagingSettings->BuildTarget = GameTargets[0].Name;
+				AllPlatformPackagingSettings->SaveConfig();
+			}
+
+			FToolMenuSection& TargetSection = Menu->AddSection("BuildTarget", LOCTEXT("TurnkeySection_BuildTarget", "Build Target"));
 
+			// Using the project default setting is represented by an empty string
+			TargetSection.AddMenuEntry(
+				NAME_None,
+				FText::Format(LOCTEXT("DefaultPackageTarget",  "Use Project Setting ({0})"), FText::FromString(AllPlatformPackagingSettings->BuildTarget)),
+				FText::Format(LOCTEXT("DefaultPackageTargetTooltip", "Package the {0} target"), FText::FromString(AllPlatformPackagingSettings->BuildTarget)),
+				FSlateIcon(),
+				FUIAction(
+					FExecuteAction::CreateStatic(&FTurnkeySupportCallbacks::SetPrepareForDebuggingBuildTarget, FString()),
+					FCanExecuteAction(),
+					FIsActionChecked::CreateStatic(&FTurnkeySupportCallbacks::PrepareForDebuggingBuildTargetIsChecked, FString())
+					),
+				EUserInterfaceActionType::RadioButton
+			);
+
+			for (const FTargetInfo& Target : ValidTargets)
+			{
 				TargetSection.AddMenuEntry(
 					NAME_None,
-					FText::Format(LOCTEXT("DefaultPackageTarget",  "Use Project Setting ({0})"), FText::FromString(AllPlatformPackagingSettings->BuildTarget)),
-					FText::Format(LOCTEXT("DefaultPackageTargetTooltip", "Package the {0} target"), FText::FromString(AllPlatformPackagingSettings->BuildTarget)),
+					FText::FromString(Target.Name),
+					FText::Format(LOCTEXT("PackageTargetName", "Package the '{0}' target."), FText::FromString(Target.Name)),
 					FSlateIcon(),
 					FUIAction(
-						FExecuteAction::CreateStatic(&FTurnkeySupportCallbacks::SetPackageBuildTarget, VanillaInfo, FString("")),
+						FExecuteAction::CreateStatic(&FTurnkeySupportCallbacks::SetPrepareForDebuggingBuildTarget, Target.Name),
 						FCanExecuteAction(),
-						FIsActionChecked::CreateStatic(&FTurnkeySupportCallbacks::PackageBuildTargetIsChecked, VanillaInfo, FString(""))
-						),
+						FIsActionChecked::CreateStatic(&FTurnkeySupportCallbacks::PrepareForDebuggingBuildTargetIsChecked, Target.Name)
+					),
 					EUserInterfaceActionType::RadioButton
 				);
+			}
+		}
+	}
+}
+
+static void MakePerCustomBuildMenu(UToolMenu* ToolMenu, ECustomBuildType Type, FProjectBuildSettings CustomBuild, FText ToolTip)
+{
+	const bool bHasDeviceEntry = CustomBuild.BuildCookRunParams.Contains(TEXT("{DeviceId}"));
+
+	FToolMenuSection& CommandSection = ToolMenu->AddSection("CommandSection");
+	{
+		CommandSection.AddMenuEntry(
+			"ExecuteCustomBuild",
+			FText::FromString(CustomBuild.Name),
+			ToolTip,
+			FSlateIcon(FAppStyle::GetAppStyleSetName(), "Icons.Advanced"),
+			FUIAction(
+				FExecuteAction::CreateStatic(&FTurnkeySupportCallbacks::ExecuteCustomBuild, Type, CustomBuild),
+				FCanExecuteAction::CreateStatic(&FTurnkeySupportCallbacks::CanExecuteCustomBuild, Type, CustomBuild.Name, bHasDeviceEntry)
+			)
+		);
+	}
 
-				for (const FTargetInfo& Target : ValidTargets)
+	// Custom builds either show a list of platforms or a list of devices
+	if (bHasDeviceEntry)
+	{
+		FToolMenuSection& DevicesSection = ToolMenu->AddSection("DevicesSection", LOCTEXT("TurnkeySection_Device", "Device"));
+
+		ITargetDeviceServicesModule* TargetDeviceServicesModule = static_cast<ITargetDeviceServicesModule*>(FModuleManager::Get().LoadModule(TEXT("TargetDeviceServices")));
+
+		TArray<FName> SupportedPlatforms;
+		GetTurnkeyCategorizedPlatforms(&SupportedPlatforms, nullptr, nullptr);
+
+		for (const FName& IniPlatformName : SupportedPlatforms)
+		{
+			if (CustomBuild.SpecificPlatforms.IsEmpty() || CustomBuild.SpecificPlatforms.Contains(IniPlatformName))
+			{
+				TArray<TSharedPtr<ITargetDeviceProxy>> DeviceProxies;
+				TargetDeviceServicesModule->GetDeviceProxyManager()->GetAllProxies(IniPlatformName, DeviceProxies);
+
+				for (const TSharedPtr<ITargetDeviceProxy>& Proxy : DeviceProxies)
 				{
-					TargetSection.AddMenuEntry(
+					FString DeviceName = Proxy->GetName();
+					FString DeviceId = Proxy->GetTargetDeviceId(NAME_None);
+
+					DevicesSection.AddMenuEntry(
 						NAME_None,
-						FText::FromString(Target.Name),
-						FText::Format(LOCTEXT("PackageTargetName", "Package the '{0}' target."), FText::FromString(Target.Name)),
-						FSlateIcon(),
+						MakeSdkStatusAttribute(IniPlatformName, Proxy),
+						FText::Format(LOCTEXT("TurnkeyTooltip_CustomBuildDevice", "{0} ({1})"), FText::FromString(DeviceName), FText::FromName(IniPlatformName)),
+						MakePlatformSdkIconAttribute(IniPlatformName, Proxy),
 						FUIAction(
-							FExecuteAction::CreateStatic(&FTurnkeySupportCallbacks::SetPackageBuildTarget, VanillaInfo, Target.Name),
+							FExecuteAction::CreateStatic(&FTurnkeySupportCallbacks::SetDeviceIdForCustomBuild, Type, CustomBuild.Name, IniPlatformName, DeviceId),
 							FCanExecuteAction(),
-							FIsActionChecked::CreateStatic(&FTurnkeySupportCallbacks::PackageBuildTargetIsChecked, VanillaInfo, Target.Name)
+							FIsActionChecked::CreateStatic(&FTurnkeySupportCallbacks::DeviceIdForCustomBuildIsChecked, Type, CustomBuild.Name, DeviceId)
 						),
 						EUserInterfaceActionType::RadioButton
 					);
 				}
 			}
 		}
-
-		if (TargetPlatform != nullptr)
+	}
+	else
+	{
+		FToolMenuSection& PlatformsSection = ToolMenu->AddSection("PlatformsSection", LOCTEXT("TurnkeySection_Platform", "Platform"));
 		{
-			// gather all possible architectures
-			TArray<FString> Architectures;
-			TargetPlatform->GetPossibleArchitectures(Architectures);
-			
-			if (Architectures.Num() > 1)
+			TArray<FName> SupportedPlatforms;
+			GetTurnkeyCategorizedPlatforms(&SupportedPlatforms, nullptr, nullptr);
+
+			for (const FName& IniPlatformName : SupportedPlatforms)
 			{
-				FToolMenuSection& ArchitectureSection = ToolMenu->AddSection("ArchitectureSelection", LOCTEXT("TurnkeySection_ArchitectureSelection", "Architecture Selection"));
+				if (CustomBuild.SpecificPlatforms.IsEmpty() || CustomBuild.SpecificPlatforms.Contains(IniPlatformName))
+				{
+					const PlatformInfo::FTargetPlatformInfo* VanillaInfo = PlatformInfo::FindVanillaPlatformInfo(IniPlatformName);
+
+					const ITargetPlatform* TargetPlatform = GetTargetPlatformManager()->FindTargetPlatform(IniPlatformName);
+
+					// Flavor selection exists for Android to be able to package ASTC, DXT, ETC2, etc
+					TArray<const PlatformInfo::FTargetPlatformInfo*> ValidFlavors;
+
+					// Architecture selection exists for platforms targeting multiple CPU architectures, such as Windows on
+					// x64 and ARM64.
+					TArray<FString> Architectures;
+
+					if (VanillaInfo != nullptr)
+					{
+						ValidFlavors = VanillaInfo->Flavors.FilterByPredicate([](const PlatformInfo::FTargetPlatformInfo* Target)
+							{
+								// Editor isn't a valid platform type that users can target
+								// The Build Target will choose client or server, so no need to show them as well
+								return Target->PlatformType != EBuildTargetType::Editor && Target->PlatformType != EBuildTargetType::Client && Target->PlatformType != EBuildTargetType::Server;
+							});
+					}
+
+					if (TargetPlatform != nullptr)
+					{
+						TargetPlatform->GetPossibleArchitectures(Architectures);
+					}
+				
+					// Platforms that have multiple flavors or architectures are displayed as a radio button that also pops
+					// out a sub-menu to choose the flavor and/or architecture.
+					if (ValidFlavors.Num() > 1 || Architectures.Num() > 1)
+					{
+						FToolMenuEntry Entry = FToolMenuEntry::InitSubMenu(
+							NAME_None,
+							MakeSdkStatusAttribute(IniPlatformName, nullptr),
+							FText::FromString(IniPlatformName.ToString()),
+							FNewToolMenuDelegate::CreateStatic(&MakeTurnkeyPerPlatformMenu, IniPlatformName, ValidFlavors, Architectures),
+							FUIAction(
+								FExecuteAction::CreateStatic(&FTurnkeySupportCallbacks::SetPlatformForCustomBuild, Type, CustomBuild.Name, IniPlatformName),
+								FCanExecuteAction(),
+								FIsActionChecked::CreateStatic(&FTurnkeySupportCallbacks::PlatformForCustomBuildIsChecked, Type, CustomBuild.Name, IniPlatformName)
+							),
+							EUserInterfaceActionType::RadioButton,
+							false,
+							MakePlatformSdkIconAttribute(IniPlatformName, nullptr)
+						);
+
+						PlatformsSection.AddEntry(Entry);
+					}
+					else
+					{
+						FToolMenuEntry Entry = FToolMenuEntry::InitMenuEntry(
+							NAME_None,
+							MakeSdkStatusAttribute(IniPlatformName, nullptr),
+							FText::FromString(IniPlatformName.ToString()),
+							MakePlatformSdkIconAttribute(IniPlatformName, nullptr),
+							FUIAction(
+								FExecuteAction::CreateStatic(&FTurnkeySupportCallbacks::SetPlatformForCustomBuild, Type, CustomBuild.Name, IniPlatformName),
+								FCanExecuteAction(),
+								FIsActionChecked::CreateStatic(&FTurnkeySupportCallbacks::PlatformForCustomBuildIsChecked, Type, CustomBuild.Name, IniPlatformName)
+							),
+							EUserInterfaceActionType::RadioButton
+						);
 
-				ArchitectureSection.AddMenuEntry(
+						PlatformsSection.AddEntry(Entry);
+					}
+
+				}
+			}
+		}
+	}
+
+	UProjectPackagingSettings* AllPlatformPackagingSettings = GetMutableDefault<UProjectPackagingSettings>();
+
+	FToolMenuSection& ConfigSection = ToolMenu->AddSection("BuildConfig", LOCTEXT("TurnkeySection_BuildConfig", "Build Configuration"));
+	{
+		const UProjectPackagingSettings::FConfigurationInfo& DefaultConfigInfo = UProjectPackagingSettings::ConfigurationInfo[static_cast<int32>(AllPlatformPackagingSettings->BuildConfiguration)];
+		ConfigSection.AddMenuEntry(
+			NAME_None,
+			FText::Format(LOCTEXT("DefaultConfiguration",  "Use Project Setting ({0})"), DefaultConfigInfo.Name),
+			DefaultConfigInfo.ToolTip,
+			FSlateIcon(),
+			FUIAction(
+				FExecuteAction::CreateStatic(&FTurnkeySupportCallbacks::SetBuildConfigurationForCustomBuild, Type, CustomBuild.Name, EProjectPackagingBuildConfigurations::PPBC_MAX),
+				FCanExecuteAction(),
+				FIsActionChecked::CreateStatic(&FTurnkeySupportCallbacks::BuildConfigurationForCustomBuildIsChecked, Type, CustomBuild.Name, EProjectPackagingBuildConfigurations::PPBC_MAX)
+			),
+			EUserInterfaceActionType::RadioButton
+		);
+
+		EProjectType ProjectType = FTurnkeyEditorSupport::DoesProjectHaveCode() ? EProjectType::Code : EProjectType::Content;
+		TArray<EProjectPackagingBuildConfigurations> PackagingConfigurations = UProjectPackagingSettings::GetValidPackageConfigurations();
+
+		for (EProjectPackagingBuildConfigurations PackagingConfiguration : PackagingConfigurations)
+		{
+			const UProjectPackagingSettings::FConfigurationInfo& ConfigurationInfo = UProjectPackagingSettings::ConfigurationInfo[(int)PackagingConfiguration];
+			if (FInstalledPlatformInfo::Get().IsValid(TOptional<EBuildTargetType>(), TOptional<FString>(), ConfigurationInfo.Configuration, ProjectType, EInstalledPlatformState::Downloaded))
+			{
+				ConfigSection.AddMenuEntry(
 					NAME_None,
-					LOCTEXT("DefaultArchitectureName", "Project Default"),
-					FText(),
+					ConfigurationInfo.Name,
+					ConfigurationInfo.ToolTip,
 					FSlateIcon(),
 					FUIAction(
-						FExecuteAction::CreateStatic(&FTurnkeySupportCallbacks::SetActiveArchitecture, VanillaInfo, FString()),
-						FCanExecuteAction::CreateStatic(&FTurnkeySupportCallbacks::CanSetActiveArchitecture, VanillaInfo, FString()),
-						FIsActionChecked::CreateStatic(&FTurnkeySupportCallbacks::SetActiveArchitectureIsChecked, VanillaInfo, FString())
+						FExecuteAction::CreateStatic(&FTurnkeySupportCallbacks::SetBuildConfigurationForCustomBuild, Type, CustomBuild.Name, PackagingConfiguration),
+						FCanExecuteAction(),
+						FIsActionChecked::CreateStatic(&FTurnkeySupportCallbacks::BuildConfigurationForCustomBuildIsChecked, Type, CustomBuild.Name, PackagingConfiguration)
 					),
 					EUserInterfaceActionType::RadioButton
 				);
+			}
+		}
+	}
 
-				for (const FString& Architecture : Architectures)
-				{
-					bool bIsHostArchitecture = (TargetPlatform->IniPlatformName() == FPlatformProperties::IniPlatformName()) && (Architecture == TargetPlatform->GetHostArchitecture());
+	// Collect build targets. Content-only projects use Engine targets
+	FProjectStatus ProjectStatus;
+	bool bHasCode = IProjectManager::Get().QueryStatusForCurrentProject(ProjectStatus) && ProjectStatus.bCodeBasedProject;
 
-					FText ArchLabel = bIsHostArchitecture 
-						? FText::Format(LOCTEXT("HostArchLabel", "{0} (this platform)"), FText::FromString(Architecture)) 
-						: FText::FromString(Architecture);
+	IDesktopPlatform* DesktopPlatform = FDesktopPlatformModule::Get();
+	TArray<FTargetInfo> Targets = bHasCode ? DesktopPlatform->GetTargetsForCurrentProject() : DesktopPlatform->GetTargetsForProject(FString());
 
-					ArchitectureSection.AddMenuEntry(
-						NAME_None,
-						ArchLabel,
-						FText(),
-						FSlateIcon(),
-						FUIAction(
-							FExecuteAction::CreateStatic(&FTurnkeySupportCallbacks::SetActiveArchitecture, VanillaInfo, Architecture),
-							FCanExecuteAction::CreateStatic(&FTurnkeySupportCallbacks::CanSetActiveArchitecture, VanillaInfo, Architecture),
-							FIsActionChecked::CreateStatic(&FTurnkeySupportCallbacks::SetActiveArchitectureIsChecked, VanillaInfo, Architecture)
-						),
-						EUserInterfaceActionType::RadioButton
-					);
+	if (Targets.Num() > 0)
+	{
+		Targets.Sort([](const FTargetInfo& A, const FTargetInfo& B) { return A.Name < B.Name; });
+		
+		const TArray<FTargetInfo> ValidTargets = Targets.FilterByPredicate([](const FTargetInfo& Target)
+			{
+				return Target.Type == EBuildTargetType::Game || Target.Type == EBuildTargetType::Client || Target.Type == EBuildTargetType::Server;
+			});
+
+		if (ValidTargets.Num() > 1)
+		{
+			// Set BuildTarget to default to Game if it hasn't been set (or if it was set, but is no longer valid)
+			FString BuildTarget = AllPlatformPackagingSettings->BuildTarget;
+			if (BuildTarget.IsEmpty() ||
+			   !ValidTargets.ContainsByPredicate([&BuildTarget](const FTargetInfo& Target) { return Target.Name == BuildTarget; } ))
+			{
+				TArray<FTargetInfo> GameTargets = ValidTargets.FilterByPredicate([](const FTargetInfo& Target)
+					{ return Target.Type == EBuildTargetType::Game;	});
+				// if there are no Game targets, look for a Client target
+				if (GameTargets.Num() == 0)
+				{
+					GameTargets = ValidTargets.FilterByPredicate([](const FTargetInfo& Target)
+						{ return Target.Type == EBuildTargetType::Client;	});
+				}
+				// if no clients, then just _anything_ 
+				if (GameTargets.Num() == 0)
+				{
+					GameTargets = ValidTargets;
 				}
+
+				AllPlatformPackagingSettings->BuildTarget = GameTargets[0].Name;
+				AllPlatformPackagingSettings->SaveConfig();
+			}
+
+			FToolMenuSection& TargetSection = ToolMenu->AddSection("BuildTarget", LOCTEXT("TurnkeySection_BuildTarget", "Build Target"));
+
+			// Using the project default setting is represented by an empty string
+			TargetSection.AddMenuEntry(
+				NAME_None,
+				FText::Format(LOCTEXT("DefaultPackageTarget",  "Use Project Setting ({0})"), FText::FromString(AllPlatformPackagingSettings->BuildTarget)),
+				FText::Format(LOCTEXT("DefaultPackageTargetTooltip", "Package the {0} target"), FText::FromString(AllPlatformPackagingSettings->BuildTarget)),
+				FSlateIcon(),
+				FUIAction(
+					FExecuteAction::CreateStatic(&FTurnkeySupportCallbacks::SetBuildTargetForCustomBuild, Type, CustomBuild.Name, FString()),
+					FCanExecuteAction(),
+					FIsActionChecked::CreateStatic(&FTurnkeySupportCallbacks::BuildTargetForCustomBuildIsChecked, Type, CustomBuild.Name, FString())
+					),
+				EUserInterfaceActionType::RadioButton
+			);
+
+			for (const FTargetInfo& Target : ValidTargets)
+			{
+				TargetSection.AddMenuEntry(
+					NAME_None,
+					FText::FromString(Target.Name),
+					FText::Format(LOCTEXT("PackageTargetName", "Package the '{0}' target."), FText::FromString(Target.Name)),
+					FSlateIcon(),
+					FUIAction(
+						FExecuteAction::CreateStatic(&FTurnkeySupportCallbacks::SetBuildTargetForCustomBuild, Type, CustomBuild.Name, Target.Name),
+						FCanExecuteAction(),
+						FIsActionChecked::CreateStatic(&FTurnkeySupportCallbacks::BuildTargetForCustomBuildIsChecked, Type, CustomBuild.Name, Target.Name)
+					),
+					EUserInterfaceActionType::RadioButton
+				);
 			}
 		}
+	}
+}
+
+static void MakeCustomBuildsMenuItems(UToolMenu* Menu)
+{
+	UProjectPackagingSettings* AllPlatformPackagingSettings = GetMutableDefault<UProjectPackagingSettings>();
+
+	// Items from EngineCustomBuilds and ProjectCustomBuilds are shown in separate sections. Since they are separate
+	// arrays, it's possible they both contain an entry with the same name.
+	FToolMenuSection& EngineSection = Menu->AddSection("EngineCustomBuilds", LOCTEXT("TurnkeySection_EngineCustomBuilds", "Engine Custom Builds"));
+
+	for (FProjectBuildSettings& CustomBuild : AllPlatformPackagingSettings->EngineCustomBuilds)
+	{
+#if ALLOW_CONTROL_TO_COPY_COMMANDLINE
+		FText ToolTipFormat = LOCTEXT("TurnkeyTooltip_EngineCustomBuild_WithCopy", "Execute the '{BuildName}' command.\n{BuildHelp}\n--\nHold Control to copy the final underlying BuildCookRun commandline to the clipboard.");
+#else
+		FText ToolTipFormat = LOCTEXT("TurnkeyTooltip_EngineCustomBuild", "Execute the '{BuildName}' command.\n{BuildHelp}");
+#endif
 
+		FFormatNamedArguments Args;
+		Args.Add(TEXT("BuildName"), FText::FromString(CustomBuild.Name));
+		Args.Add(TEXT("BuildHelp"), FText::FromString(CustomBuild.HelpText));
+
+		FText ToolTip = FText::Format(ToolTipFormat, Args);
+
+		EngineSection.AddSubMenu(
+			NAME_None,
+			FText::FromString(CustomBuild.Name),
+			ToolTip,
+			FNewToolMenuDelegate::CreateStatic(&MakePerCustomBuildMenu, ECustomBuildType::EngineCustomBuilds, CustomBuild, ToolTip),
+			false,
+			FSlateIcon(FAppStyle::GetAppStyleSetName(), "Icons.Advanced"),
+			true
+		);
+	}
+
+	FToolMenuSection& ProjectSection = Menu->AddSection("ProjectCustomBuilds", LOCTEXT("TurnkeySection_ProjectCustomBuilds", "Project Custom Builds"));
+
+	for (FProjectBuildSettings& CustomBuild : AllPlatformPackagingSettings->ProjectCustomBuilds)
+	{
+#if ALLOW_CONTROL_TO_COPY_COMMANDLINE
+		FText ToolTipFormat = LOCTEXT("TurnkeyTooltip_ProjectCustomBuild_WithCopy", "Execute this project's custom '{BuildName}' command.\n{BuildHelp}\n--\nHold Control to copy the final BuildCookRun commandline to the clipboard.\nThis custom command comes from Project Settings.");
+#else
+		FText ToolTipFormat = LOCTEXT("TurnkeyTooltip_ProjectCustomBuild", "Execute this project's custom '{BuildName}' command.\n{BuildHelp}\n--\nThis custom command comes from Project Settings.");
+#endif
+
+		FFormatNamedArguments Args;
+		Args.Add(TEXT("BuildName"), FText::FromString(CustomBuild.Name));
+		Args.Add(TEXT("BuildHelp"), FText::FromString(CustomBuild.HelpText));
+
+		FText ToolTip = FText::Format(ToolTipFormat, Args);
+
+		ProjectSection.AddSubMenu(
+			NAME_None,
+			FText::FromString(CustomBuild.Name),
+			ToolTip,
+			FNewToolMenuDelegate::CreateStatic(&MakePerCustomBuildMenu, ECustomBuildType::ProjectCustomBuilds, CustomBuild, ToolTip),
+			false,
+			FSlateIcon(FAppStyle::GetAppStyleSetName(), "Icons.Advanced"),
+			true
+		);
+	}
+}
+
+static void MakeSdkManagementPerPlatformPlatformMenu(UToolMenu* ToolMenu, FName IniPlatformName, ITargetDeviceServicesModule* TargetDeviceServicesModule)
+{
+	const FDataDrivenPlatformInfo& DDPI = FDataDrivenPlatformInfoRegistry::GetPlatformInfo(IniPlatformName);
+	FString UBTPlatformString = DDPI.UBTPlatformString;
+
+	const PlatformInfo::FTargetPlatformInfo* VanillaInfo = PlatformInfo::FindVanillaPlatformInfo(IniPlatformName);
+
+	const ITargetPlatform* TargetPlatform = GetTargetPlatformManager()->FindTargetPlatform(IniPlatformName);
+
+	if (VanillaInfo != nullptr)
+	{
+		UProjectPackagingSettings* PackagingSettings = FTurnkeySupportCallbacks::GetPackagingSettingsForPlatform(IniPlatformName);
+
+		UProjectPackagingSettings* AllPlatformPackagingSettings = GetMutableDefault<UProjectPackagingSettings>();
 
 		FToolMenuSection& DevicesSection = ToolMenu->AddSection("AllDevices", LOCTEXT("TurnkeySection_AllDevices", "All Devices"));
 
@@ -1515,28 +2357,6 @@ static void MakeTurnkeyPlatformMenu(UToolMenu* ToolMenu, FName IniPlatformName,
 				FText(),
 				FNewToolMenuDelegate::CreateLambda([IniPlatformName, DeviceName, DeviceId, PackagingSettings](UToolMenu* SubToolMenu)
 				{
-					if (HasBuildForDeviceOrNot(PackagingSettings->EngineCustomBuilds, true) || HasBuildForDeviceOrNot(PackagingSettings->ProjectCustomBuilds, true))
-					{
-						FToolMenuSection& CustomBuildSection = SubToolMenu->AddSection("DeviceCustomBuilds", LOCTEXT("TurnkeySection_DeviceCustomBuilds", "Builds For Devices"));
-
-						// these ToolTipFormats will be formatted inside the MakeCustomBuildMenuEntries function with some named args
-						// the empty strings passed in is the deviceId, which means to show builds that don't contain DeviceId components
-#if ALLOW_CONTROL_TO_COPY_COMMANDLINE
-						FText ToolTipFormat = LOCTEXT("TurnkeyTooltip_EngineCustomBuild_WithCopyDevice", "Execute the '{BuildName}' command on {DeviceName}.\n{BuildHelp}\n--\nThis runs the following command:\nRunUAT {CommandLine}\nHold Control to copy the final underlying BuildCookRun commandline to the clipboard.");
-#else
-						FText ToolTipFormat = LOCTEXT("TurnkeyTooltip_EngineCustomBuildDevice", "Execute the '{BuildName}' command on {DeviceName}.\n{BuildHelp}\n--\nThis runs the following command:\nRunUAT {CommandLine}");
-#endif
-						MakeCustomBuildMenuEntries(CustomBuildSection, PackagingSettings->EngineCustomBuilds, IniPlatformName, DeviceId, DeviceName, ToolTipFormat);
-
-#if ALLOW_CONTROL_TO_COPY_COMMANDLINE
-						ToolTipFormat = LOCTEXT("TurnkeyTooltip_ProjectCustomBuildDevice", "Execute this project's custom '{BuildName}' command on {DeviceName}.\n{BuildHelp}\n--\nThis runs the following command:\nRunUAT {CommandLine}\nHold Control to copy the final BuildCookRun commandline to the clipboard.\nThis custom command comes from Project Settings.");
-#else
-						ToolTipFormat = LOCTEXT("TurnkeyTooltip_ProjectCustomBuild_WithCopyDevice", "Execute this project's custom '{BuildName}' command on {DeviceName}.\n{BuildHelp}\n--\nThis runs the following command:\nRunUAT {CommandLine}\nThis custom command comes from Project Settings.");
-#endif
-						MakeCustomBuildMenuEntries(CustomBuildSection, PackagingSettings->ProjectCustomBuilds, IniPlatformName, DeviceId, DeviceName, ToolTipFormat);
-					}
-
-
 					FTurnkeySdkInfo SdkInfo = ITurnkeySupportModule::Get().GetSdkInfoForDeviceId(DeviceId);
 					FText SdkText, SdkToolTip;
 					FormatSdkInfo(DeviceId, SdkInfo, SdkText, SdkToolTip);
@@ -1685,20 +2505,99 @@ static void MakeTurnkeyPlatformMenu(UToolMenu* ToolMenu, FName IniPlatformName,
 				FExecuteAction::CreateStatic(TurnkeyInstallSdk, IniPlatformName.ToString(), true, false, NoDevice)
 			);
 		}
-	}
+	}
+
+#if 0 // @todo turnkey enable and always show documentation for installation help when we have URLs in place
+	// Link to documentation
+	SdkSection.AddMenuEntry(
+		NAME_None,
+		LOCTEXT("Turnkey_ShowDocumentation", "Installation Help..."),
+		LOCTEXT("TurnkeyTooltip_ShowDocumentation", "Show documentation with help installing the SDK for this platform"),
+		FSlateIcon(),
+		FExecuteAction::CreateStatic(ShowInstallationHelp, IniPlatformName)
+	);
+#endif
+}
+
+
+static void MakeSdkManagementMenuItems(UToolMenu* SdkManagementMenu)
+{
+	SdkManagementMenu->AddDynamicSection(NAME_None, FNewToolMenuDelegate::CreateLambda([](UToolMenu* Menu)
+	{
+		FToolMenuSection& SupportedPlatformsSection = Menu->AddSection("SupportedPlatforms", LOCTEXT("TurnkeySection_SupportedPlatforms", "Platforms"));
+
+		ITargetDeviceServicesModule* TargetDeviceServicesModule = static_cast<ITargetDeviceServicesModule*>(FModuleManager::Get().LoadModule(TEXT("TargetDeviceServices")));
+
+		TArray<FName> SupportedPlatforms;
+		TArray<FName> UnsupportedPlatforms;
+		TArray<FName> UncompiledPlatforms;
+		GetTurnkeyCategorizedPlatforms(&SupportedPlatforms, &UnsupportedPlatforms, &UncompiledPlatforms);
+
+		for (const FName& IniPlatformName : SupportedPlatforms)
+		{
+			SupportedPlatformsSection.AddSubMenu(
+				NAME_None,
+				MakeSdkStatusAttribute(IniPlatformName, nullptr),
+				FText::FromString(IniPlatformName.ToString()),
+				FNewToolMenuDelegate::CreateStatic(&MakeSdkManagementPerPlatformPlatformMenu, IniPlatformName, TargetDeviceServicesModule),
+				false,
+				MakePlatformSdkIconAttribute(IniPlatformName, nullptr),
+				true
+			);
+		}
+
+		if (UnsupportedPlatforms.Num() != 0)
+		{
+			FToolMenuSection& UnsupportedPlatformsSection = Menu->AddSection("UnsupportedPlatforms", LOCTEXT("Turnkey_UnsupportedPlatforms", "Unsupported Platforms"));
+
+			for (const FName& IniPlatformName : UnsupportedPlatforms)
+			{
+				UnsupportedPlatformsSection.AddSubMenu(
+					NAME_None,
+					MakeSdkStatusAttribute(IniPlatformName, nullptr),
+					FText::FromString(IniPlatformName.ToString()),
+					FNewToolMenuDelegate::CreateStatic(&MakeSdkManagementPerPlatformPlatformMenu, IniPlatformName, TargetDeviceServicesModule),
+					false,
+					MakePlatformSdkIconAttribute(IniPlatformName, nullptr),
+					true
+				);
+			}
+		}
 
-#if 0 // @todo turnkey enable and always show documentation for installation help when we have URLs in place
-	// Link to documentation
-	SdkSection.AddMenuEntry(
-		NAME_None,
-		LOCTEXT("Turnkey_ShowDocumentation", "Installation Help..."),
-		LOCTEXT("TurnkeyTooltip_ShowDocumentation", "Show documentation with help installing the SDK for this platform"),
-		FSlateIcon(),
-		FExecuteAction::CreateStatic(ShowInstallationHelp, IniPlatformName)
-	);
-#endif
-}
+		if (UncompiledPlatforms.Num() != 0)
+		{
+			FToolMenuSection& UncompiledPlatformsSection = Menu->AddSection("UncompiledPlatforms", LOCTEXT("Turnkey_UncompiledPlatforms", "Platforms With No Compiled Support"));
+
+			for (const FName& IniPlatformName : UncompiledPlatforms)
+			{
+				UncompiledPlatformsSection.AddSubMenu(
+					NAME_None,
+					MakeSdkStatusAttribute(IniPlatformName, nullptr),
+					FText::FromString(IniPlatformName.ToString()),
+					FNewToolMenuDelegate::CreateStatic(&MakeSdkManagementPerPlatformPlatformMenu, IniPlatformName, TargetDeviceServicesModule),
+					false,
+					MakePlatformSdkIconAttribute(IniPlatformName, nullptr),
+					true
+				);
+			}
+		}
+
+		FToolMenuSection& RefreshSection = Menu->AddSection("RefreshSection");
 
+		RefreshSection.AddSeparator(NAME_None);
+		
+		RefreshSection.AddMenuEntry(
+			NAME_None,
+			LOCTEXT("RefreshPlatformStatus", "Refresh platform status"),
+			LOCTEXT("RefreshPlatformStatusDescription", "Update all platforms and device information"),
+			FSlateIcon(FAppStyle::GetAppStyleSetName(), "SourceControl.Actions.Refresh"),
+			FUIAction(
+				FExecuteAction::CreateStatic(&FTurnkeySupportCallbacks::SetRefreshPlatformStatus),
+				FCanExecuteAction::CreateStatic(&FTurnkeySupportCallbacks::CanRefreshPlatformStatus)
+			)
+		);
+	}));
+}
 
 // Launch On
 
@@ -1895,7 +2794,7 @@ static void GenerateDeviceProxyMenuParams(TSharedPtr<ITargetDeviceProxy> DeviceP
 
 void FTurnkeySupportModule::MakeQuickLaunchItems(class UToolMenu* Menu, FOnQuickLaunchSelected ExternalOnClickDelegate) const
 {
-	FToolMenuSection& MenuSection = Menu->AddSection("QuickLaunchDevices", LOCTEXT("QuickLaunch", "Quick Launch"));
+	FToolMenuSection& MenuSection = Menu->AddSection("QuickLaunchDevices", LOCTEXT("Devices", "Devices"));
 
 	MenuSection.AddDynamicEntry("PlatformsMenu", FNewToolMenuSectionDelegate::CreateLambda([ExternalOnClickDelegate](FToolMenuSection& DynamicSection)
 	{
@@ -2037,24 +2936,6 @@ void FTurnkeySupportModule::MakeQuickLaunchItems(class UToolMenu* Menu, FOnQuick
 				}
 			}
 		}
-
-		// if we don't have an external delegate to call, then this is the internally included items in the Platforms menu and we can add the extra option(s)
-		if (!ExternalOnClickDelegate.IsBound())
-		{
-			DynamicSection.AddMenuEntry(
-				NAME_None,
-				LOCTEXT("CookOnTheFlyOnLaunch", "Enable cooking on the fly"),
-				LOCTEXT("CookOnTheFlyOnLaunchDescription", "Cook on the fly instead of cooking upfront when launching"),
-				FSlateIcon(FAppStyle::GetAppStyleSetName(), "MainFrame.CookContent"),
-				FUIAction(
-					FExecuteAction::CreateStatic(&FTurnkeySupportCallbacks::SetCookOnTheFly),
-					FCanExecuteAction::CreateStatic(&FTurnkeySupportCallbacks::CanSetCookOnTheFly),
-					FIsActionChecked::CreateStatic(&FTurnkeySupportCallbacks::SetCookOnTheFlyIsChecked)
-				),
-				EUserInterfaceActionType::ToggleButton
-			);
-		}
-
 		// now kick-off any devices that need to be updated
 		if (DeviceIdsToQuery.Num() > 0)
 		{
@@ -2063,6 +2944,49 @@ void FTurnkeySupportModule::MakeQuickLaunchItems(class UToolMenu* Menu, FOnQuick
 	}
 	));
 
+
+	UProjectPackagingSettings* AllPlatformPackagingSettings = GetMutableDefault<UProjectPackagingSettings>();
+
+	FToolMenuSection& ConfigSection = Menu->AddSection("BuildConfig", LOCTEXT("TurnkeySection_BuildConfig", "Build Configuration"));
+	{
+		const UProjectPackagingSettings::FConfigurationInfo& DefaultConfigInfo = UProjectPackagingSettings::ConfigurationInfo[static_cast<int32>(AllPlatformPackagingSettings->BuildConfiguration)];
+		ConfigSection.AddMenuEntry(
+			NAME_None,
+			FText::Format(LOCTEXT("DefaultConfiguration",  "Use Project Setting ({0})"), DefaultConfigInfo.Name),
+			DefaultConfigInfo.ToolTip,
+			FSlateIcon(),
+			FUIAction(
+				FExecuteAction::CreateStatic(&FTurnkeySupportCallbacks::SetLaunchOnBuildConfiguration, EProjectPackagingBuildConfigurations::PPBC_MAX),
+				FCanExecuteAction(),
+				FIsActionChecked::CreateStatic(&FTurnkeySupportCallbacks::LaunchOnBuildConfigurationIsChecked, EProjectPackagingBuildConfigurations::PPBC_MAX)
+			),
+			EUserInterfaceActionType::RadioButton
+		);
+
+		EProjectType ProjectType = FTurnkeyEditorSupport::DoesProjectHaveCode() ? EProjectType::Code : EProjectType::Content;
+		TArray<EProjectPackagingBuildConfigurations> PackagingConfigurations = UProjectPackagingSettings::GetValidPackageConfigurations();
+
+		for (EProjectPackagingBuildConfigurations PackagingConfiguration : PackagingConfigurations)
+		{
+			const UProjectPackagingSettings::FConfigurationInfo& ConfigurationInfo = UProjectPackagingSettings::ConfigurationInfo[(int)PackagingConfiguration];
+			if (FInstalledPlatformInfo::Get().IsValid(TOptional<EBuildTargetType>(), TOptional<FString>(), ConfigurationInfo.Configuration, ProjectType, EInstalledPlatformState::Downloaded))
+			{
+				ConfigSection.AddMenuEntry(
+					NAME_None,
+					ConfigurationInfo.Name,
+					ConfigurationInfo.ToolTip,
+					FSlateIcon(),
+					FUIAction(
+						FExecuteAction::CreateStatic(&FTurnkeySupportCallbacks::SetLaunchOnBuildConfiguration, PackagingConfiguration),
+						FCanExecuteAction(),
+						FIsActionChecked::CreateStatic(&FTurnkeySupportCallbacks::LaunchOnBuildConfigurationIsChecked, PackagingConfiguration)
+					),
+					EUserInterfaceActionType::RadioButton
+				);
+			}
+		}
+	}
+
 	FProjectStatus ProjectStatus;
 	bool bHasCode = IProjectManager::Get().QueryStatusForCurrentProject(ProjectStatus) && ProjectStatus.bCodeBasedProject;
 
@@ -2095,150 +3019,26 @@ void FTurnkeySupportModule::MakeQuickLaunchItems(class UToolMenu* Menu, FOnQuick
 
 	}
 
-}
-
-static void MakeSimulatorMenu(const FString& ModelType, const TArray<TSharedPtr<ITargetDeviceProxy>>& DeviceProxies, FOnQuickLaunchSelected ExternalOnClickDelegate, FToolMenuSection& DynamicSection)
-{
-	FName PlatformName = TEXT("IOS");
-	FName SubPlatformName = *ModelType;
-	FString LastChosen;
-	const TSharedPtr<ITargetDeviceProxy>* LastChosenProxy;
-	if (PerPlatformLastSimChosen.Contains(SubPlatformName))
-	{
-		LastChosen = PerPlatformLastSimChosen[SubPlatformName];
-		const TSharedPtr<ITargetDeviceProxy>* Proxy = DeviceProxies.FindByPredicate([LastChosen](const TSharedPtr<ITargetDeviceProxy>& InProxy)
-		{
-			return InProxy->GetName() == LastChosen;
-		});
-		if (Proxy)
-		{
-			LastChosenProxy = Proxy;
-		}
-		else
-		{
-			return;
-		}
-	}
-	else
-	{
-		const TSharedPtr<ITargetDeviceProxy>* FirstProxy = DeviceProxies.FindByPredicate([ModelType](const TSharedPtr<ITargetDeviceProxy>& InProxy)
-		{
-			return InProxy->GetType() == ModelType;
-		});
-		if (FirstProxy)
-		{
-			LastChosen = (*FirstProxy)->GetName();
-			LastChosenProxy = FirstProxy;
-			PerPlatformLastSimChosen.Add(SubPlatformName, LastChosen);
-		}
-		else
-		{
-			// There is no simulator of this model type to display, so return
-			return;
-		}
-	}
-	
-	// always use the first one, after sorting
-	FUIAction Action;
-	FText Tooltip;
-	
-	ITargetDeviceServicesModule* TargetDeviceServicesModule = static_cast<ITargetDeviceServicesModule*>(FModuleManager::Get().LoadModule(TEXT("TargetDeviceServices")));
-
-	GenerateDeviceProxyMenuParams(*LastChosenProxy, PlatformName, Action, Tooltip, ExternalOnClickDelegate);
-
-	const ITargetPlatform* Platform = GetTargetPlatformManager()->FindTargetPlatform(PlatformName.ToString());
-	DynamicSection.AddSubMenu(
-		NAME_None,
-		MakeSdkStatusAttribute(PlatformName, *LastChosenProxy),
-		Tooltip,
-		FNewToolMenuDelegate::CreateLambda([TargetDeviceServicesModule, ModelType, PlatformName, LastChosen, DeviceProxies, ExternalOnClickDelegate](UToolMenu* SubToolMenu)
-		{
-			FToolMenuSection& Section = SubToolMenu->AddSection(NAME_None);
-			// for each one, put an entry 
-			for (const TSharedPtr<ITargetDeviceProxy>& Proxy : DeviceProxies)
-			{
-				if (Proxy->GetType() != ModelType)
-				{
-					continue;
-				}
-				
-				// Skip over the top level menu item
-				if (LastChosen == Proxy->GetName())
-				{
-					continue;
-				}
-				
-				FUIAction SubAction;
-				FText SubTooltip;
-				GenerateDeviceProxyMenuParams(Proxy, PlatformName, SubAction, SubTooltip, ExternalOnClickDelegate);
-				Section.AddMenuEntry(NAME_None,
-									 MakeSdkStatusAttribute(PlatformName, Proxy),
-									 SubTooltip,
-									 MakePlatformSdkIconAttribute(PlatformName, Proxy),
-									 SubAction,
-									 EUserInterfaceActionType::Button);
-			}
-		}),
-		Action,
-		EUserInterfaceActionType::Check,
-		false,
-		MakePlatformSdkIconAttribute(PlatformName, *LastChosenProxy),
-		true
-	);
-	
-}
-
-void FTurnkeySupportModule::MakeSimulatorItems(class UToolMenu* Menu, FOnQuickLaunchSelected ExternalOnClickDelegate) const
-{
-	FToolMenuSection& MenuSection = Menu->AddSection("SimulatorDevices", LOCTEXT("Simulators", "Simulators"));
-
-	MenuSection.AddDynamicEntry("PlatformsMenu", FNewToolMenuSectionDelegate::CreateLambda([ExternalOnClickDelegate](FToolMenuSection& DynamicSection)
+	// if we don't have an external delegate to call, then this is the internally included items in the Platforms menu and we can add the extra option(s)
+	if (!ExternalOnClickDelegate.IsBound())
 	{
-		TArray<FString> DeviceIdsToQuery;
-		ITargetDeviceServicesModule* TargetDeviceServicesModule = static_cast<ITargetDeviceServicesModule*>(FModuleManager::Get().LoadModule(TEXT("TargetDeviceServices")));
-		for (const auto& Pair : FDataDrivenPlatformInfoRegistry::GetAllPlatformInfos())
-		{
-			if (Pair.Value.bIsFakePlatform)
-			{
-				continue;
-			}
-
-			FName PlatformName = Pair.Key;
-			const FDataDrivenPlatformInfo& Info = Pair.Value;
-			if (FDataDrivenPlatformInfoRegistry::IsPlatformHiddenFromUI(PlatformName))
-			{
-				continue;
-			}
+		FToolMenuSection& CookOnTheFlySection = Menu->AddSection("CookOnTheFlySection");
 
-			// look for devices for all platforms, even if the platform isn't installed - Turnkey can install Sdk after selecting LaunchOn
-			TArray<TSharedPtr<ITargetDeviceProxy>> AllDeviceProxies;
-			TargetDeviceServicesModule->GetDeviceProxyManager()->GetAllProxies(PlatformName, AllDeviceProxies);
-
-			// We only care about Simulator based devices
-			TArray<TSharedPtr<ITargetDeviceProxy>> DeviceProxies;
-			DeviceProxies = AllDeviceProxies.FilterByPredicate([](const TSharedPtr<ITargetDeviceProxy>& Device)
-			{
-				return Device->GetConnectionType() == TEXT("Simulator");
-			});
-			
-			if (DeviceProxies.Num() > 0)
-			{
-				// Partition based on device type
-				Algo::Sort(DeviceProxies, [](TSharedPtr<ITargetDeviceProxy> A, TSharedPtr<ITargetDeviceProxy> B)
-						   { return A->GetModel() >= B->GetModel(); });
-				
-				MakeSimulatorMenu(TEXT("Phone"), DeviceProxies, ExternalOnClickDelegate, DynamicSection);
-				MakeSimulatorMenu(TEXT("Tablet"), DeviceProxies, ExternalOnClickDelegate, DynamicSection);
-			}
-		}
+		CookOnTheFlySection.AddSeparator(NAME_None);
 
-		// now kick-off any devices that need to be updated
-		if (DeviceIdsToQuery.Num() > 0)
-		{
-			ITurnkeySupportModule::Get().UpdateSdkInfoForDevices(DeviceIdsToQuery);
-		}
+		CookOnTheFlySection.AddMenuEntry(
+			NAME_None,
+			LOCTEXT("CookOnTheFlyOnLaunch", "Cook on the Fly"),
+			LOCTEXT("CookOnTheFlyOnLaunchDescription", "Cook on the fly instead of cooking upfront when launching"),
+			FSlateIcon(FAppStyle::GetAppStyleSetName(), "MainFrame.CookContent"),
+			FUIAction(
+				FExecuteAction::CreateStatic(&FTurnkeySupportCallbacks::SetCookOnTheFly),
+				FCanExecuteAction(),
+				FIsActionChecked::CreateStatic(&FTurnkeySupportCallbacks::SetCookOnTheFlyIsChecked)
+			),
+			EUserInterfaceActionType::ToggleButton
+		);
 	}
-	));
 }
 
 TSharedRef<SWidget> FTurnkeySupportModule::MakeTurnkeyMenuWidget() const
@@ -2252,151 +3052,149 @@ TSharedRef<SWidget> FTurnkeySupportModule::MakeTurnkeyMenuWidget() const
 	{
 		UToolMenu* Menu = UToolMenus::Get()->RegisterMenu(MenuName);
 
-		FOnQuickLaunchSelected EmptyFunc;
-		MakeQuickLaunchItems(Menu, EmptyFunc);
-		MakeSimulatorItems(Menu, EmptyFunc);
-
-		// need to make this dyamic so icons, etc can update with SDK 
-		// shared devices section
-
-		FToolMenuSection& ManagePlatformsSection = Menu->AddSection("AllPlatforms", LOCTEXT("TurnkeyMenu_ManagePlatforms", "Content/Sdk/Device Management"));
-		ManagePlatformsSection.AddDynamicEntry(NAME_None, FNewToolMenuSectionDelegate::CreateLambda([](FToolMenuSection& PlatformsSection)
+		FToolMenuSection& LaunchSection = Menu->AddSection("Launch", LOCTEXT("TurnkeySection_Launch", "Launch"));
 		{
-			ITargetDeviceServicesModule* TargetDeviceServicesModule = static_cast<ITargetDeviceServicesModule*>(FModuleManager::Get().LoadModule(TEXT("TargetDeviceServices")));
-
-			TMap<FName, const FDataDrivenPlatformInfo*> UncompiledPlatforms;
-			TMap<FName, const FDataDrivenPlatformInfo*> UnsupportedPlatforms;
-
-			FProjectStatus ProjectStatus;
-			bool bProjectStatusIsValid = IProjectManager::Get().QueryStatusForCurrentProject(ProjectStatus);
+			LaunchSection.AddSubMenu(
+				"QuickLaunchSubMenu",
+				LOCTEXT("QuickLaunch", "Quick Launch Current Level"),
+				LOCTEXT("QuickLaunch_ToolTip", "Launch only the current level on the device for fast iteration"),
+				FNewToolMenuDelegate::CreateLambda([this](UToolMenu* SubMenu)
+					{
+						FOnQuickLaunchSelected EmptyFunc;
+						MakeQuickLaunchItems(SubMenu, EmptyFunc);
+					}
+				),
+				false, /* bInOpenSubMenuOnClick */
+				FSlateIcon(FAppStyle::Get().GetStyleSetName(), "Launcher.TabIcon")
+			);
+		}
 
-			for (const auto& Pair : FDataDrivenPlatformInfoRegistry::GetAllPlatformInfos())
+		FToolMenuSection& ProjectLauncherSection = Menu->AddSection("ProjectLauncher", LOCTEXT("TurnkeySection_ProjectLauncher", "Project Launcher"));
+		{
+			if (IPluginManager::Get().FindEnabledPlugin(TEXT("ProjectLauncher")) != nullptr)
 			{
-				if (Pair.Value.bIsFakePlatform || Pair.Value.bEnabledForUse == false)
-				{
-					continue;
-				}
-
-				FName PlatformName = Pair.Key;
-				const FDataDrivenPlatformInfo& Info = Pair.Value;
-				if (FDataDrivenPlatformInfoRegistry::IsPlatformHiddenFromUI(PlatformName))
-				{
-					continue;
-				}
-
-				if (!FDataDrivenPlatformInfoRegistry::HasCompiledSupportForPlatform(PlatformName, FDataDrivenPlatformInfoRegistry::EPlatformNameType::Ini))
-				{	
-					UncompiledPlatforms.Add(PlatformName, &Info);
-					continue;
-				}
-
-				if (bProjectStatusIsValid && !ProjectStatus.IsTargetPlatformSupported(PlatformName))
-				{
-					UnsupportedPlatforms.Add(PlatformName, &Info);
-					continue;
-				}
-
-				PlatformsSection.AddSubMenu(
+				ProjectLauncherSection.AddMenuEntry(
 					NAME_None,
-					MakeSdkStatusAttribute(PlatformName, nullptr),
-					FText::FromString(PlatformName.ToString()),
-					FNewToolMenuDelegate::CreateStatic(&MakeTurnkeyPlatformMenu, PlatformName, TargetDeviceServicesModule),
-					false,
-					MakePlatformSdkIconAttribute(PlatformName, nullptr),
-					true
+					LOCTEXT("OpenProjectLauncher", "Project Launcher..."),
+					LOCTEXT("OpenProjectLauncher_ToolTip", "Open the Project Launcher for advanced packaging, deploying and launching of your projects"),
+					FSlateIcon(FAppStyle::Get().GetStyleSetName(), "Launcher.TabIcon"),
+					FUIAction(FExecuteAction::CreateStatic(&FTurnkeySupportCallbacks::OpenProjectLauncher))
 				);
 			}
 
-			PlatformsSection.AddMenuEntry(
+			ProjectLauncherSection.AddMenuEntry(
 				NAME_None,
-				LOCTEXT("RefreshPlatformStatus", "Refresh platform status"),
-				LOCTEXT("RefreshPlatformStatusDescription", "Update all platforms and device information"),
-				FSlateIcon(FAppStyle::GetAppStyleSetName(), "SourceControl.Actions.Refresh"),
-				FUIAction(
-				FExecuteAction::CreateStatic(&FTurnkeySupportCallbacks::SetRefreshPlatformStatus),
-				FCanExecuteAction::CreateStatic(&FTurnkeySupportCallbacks::CanRefreshPlatformStatus))
+				LOCTEXT("OpenLegacyProjectLauncher", "Project Launcher (Legacy)..."),
+				LOCTEXT("OpenLegacyProjectLauncher_ToolTip", "Open the legacy Project Launcher for advanced packaging, deploying and launching of your projects"),
+				FSlateIcon(FAppStyle::Get().GetStyleSetName(), "Launcher.TabIcon"),
+				FUIAction(FExecuteAction::CreateStatic(&FTurnkeySupportCallbacks::OpenLegacyProjectLauncher))
+			);
+		}
+
+		FToolMenuSection& ContentManagementSection = Menu->AddSection("ContentManagementSection", LOCTEXT("TurnkeySection_ContentManagement", "Cook or Package"));
+		{
+			ContentManagementSection.AddSubMenu(
+				"CookSubMenu",
+				LOCTEXT("Turnkey_CookContent", "Cook Content"),
+				LOCTEXT("TurnkeyTooltip_CookContentSubMenu", "Cook this project"),
+				FNewToolMenuDelegate::CreateLambda([this](UToolMenu* SubMenu)
+					{
+						MakeCookContentMenuItems(SubMenu);
+					}
+				),
+				false,
+				FSlateIcon(FAppStyle::GetAppStyleSetName(), "MainFrame.CookContent")
+			);
+
+			ContentManagementSection.AddSubMenu(
+				"PackageSubMenu",
+				LOCTEXT("Turnkey_PackageProject", "Package Project"),
+				LOCTEXT("TurnkeyTooltip_PackageProject", "Package this project and archive it to a user-selected directory. This can then be used to install and run."),
+				FNewToolMenuDelegate::CreateLambda([this](UToolMenu* SubMenu)
+					{
+						MakePackageProjectMenuItems(SubMenu);
+					}
+				),
+				false,
+				FSlateIcon(FAppStyle::GetAppStyleSetName(), "MainFrame.PackageProject")
 			);
 
-			if (UnsupportedPlatforms.Num() != 0)
+#if PLATFORM_WINDOWS || PLATFORM_MAC
+			// Prepare For Debugging is only on iOS and tvOS platforms. The menu entry is not visible unless those
+			// platforms are supported.
 			{
-				PlatformsSection.AddSeparator(NAME_None);
+				TArray<FName> SupportedPlatforms;
+				GetTurnkeyCategorizedPlatforms(&SupportedPlatforms, nullptr, nullptr);
 
-				PlatformsSection.AddSubMenu(
-					NAME_None,
-					LOCTEXT("Turnkey_UnsupportedPlatforms", "Platforms Not Supported by Project"),
-					LOCTEXT("Turnkey_UnsupportedPlatformsToolTip", "List of platforms that are not marked as supported by this platform. Use the \"Supported Platforms...\""),
-					FNewToolMenuDelegate::CreateLambda([UnsupportedPlatforms, TargetDeviceServicesModule](UToolMenu* SubToolMenu)
-						{
-							FToolMenuSection& Section = SubToolMenu->AddSection(NAME_None);
-							for (const auto It : UnsupportedPlatforms)
+				if (SupportedPlatforms.Contains(TEXT("IOS")) || SupportedPlatforms.Contains(TEXT("TvOS")))
+				{
+					ContentManagementSection.AddSubMenu(
+						"PrepareSubMenu",
+						LOCTEXT("Turnkey_PrepareForDebugging", "Prepare For Debugging"),
+						LOCTEXT("TurnkeyTooltip_PrepareForDebugging", "Create files to help streamline the debugging pipeline on some platforms"),
+						FNewToolMenuDelegate::CreateLambda([this](UToolMenu* SubMenu)
 							{
-								Section.AddSubMenu(
-									NAME_None,
-									MakeSdkStatusAttribute(It.Key, nullptr),
-									FText::FromString(It.Key.ToString()),
-									FNewToolMenuDelegate::CreateStatic(&MakeTurnkeyPlatformMenu, It.Key, TargetDeviceServicesModule),
-									false,
-									MakePlatformSdkIconAttribute(It.Key, nullptr),
-									true
-								);
+								MakePrepareForDebuggingMenuItems(SubMenu);
 							}
-						})
-				);
+						),
+						false,
+						FSlateIcon(FAppStyle::GetAppStyleSetName(), "Debug")
+					);
+				}
 			}
+#endif
+		}
 
-			if (UncompiledPlatforms.Num() != 0)
+		// Custom builds menu is only displayed when custom builds have been configured.
+		FToolMenuSection& CustomBuildSection = Menu->AddSection("CustomBuildsSection", LOCTEXT("TurnkeySection_CustomBuilds", "Custom Builds"));
+		{
+			TArray<FName> SupportedPlatforms;
+			GetTurnkeyCategorizedPlatforms(&SupportedPlatforms, nullptr, nullptr);
+
+			bool bHasCustomBuilds = false;
+
+			for (const FName& IniPlatformName : SupportedPlatforms)
 			{
-				PlatformsSection.AddSeparator(NAME_None);
+				UProjectPackagingSettings* PackagingSettings = FTurnkeySupportCallbacks::GetPackagingSettingsForPlatform(IniPlatformName);
 
-				PlatformsSection.AddSubMenu(
-					NAME_None,
-					LOCTEXT("Turnkey_UncompiledPlatforms", "Platforms With No Compiled Support"),
-					LOCTEXT("Turnkey_UncompiledPlatformsToolTip", "List of platforms that you have access to, but support is not compiled in to the editor. It may be caused by missing an SDK, so you attempt to install an SDK here."),
-					FNewToolMenuDelegate::CreateLambda([UncompiledPlatforms, TargetDeviceServicesModule](UToolMenu* SubToolMenu)
-						{
-							FToolMenuSection& Section = SubToolMenu->AddSection(NAME_None);
-							for (const auto It : UncompiledPlatforms)
-							{
-								Section.AddSubMenu(
-									NAME_None,
-									MakeSdkStatusAttribute(It.Key, nullptr),
-									FText::FromString(It.Key.ToString()),
-									FNewToolMenuDelegate::CreateStatic(&MakeTurnkeyPlatformMenu, It.Key, TargetDeviceServicesModule),
-									false,
-									MakePlatformSdkIconAttribute(It.Key, nullptr),
-									true
-								);
-							}
-						})
-				);
+				if (PackagingSettings->EngineCustomBuilds.IsEmpty() || !PackagingSettings->ProjectCustomBuilds.IsEmpty())
+				{
+					bHasCustomBuilds = true;
+					break;
+				}
 			}
-		}));
 
-
-		// options section
-		FToolMenuSection& OptionsSection = Menu->AddSection("TurnkeyOptions", LOCTEXT("TurnkeySection_Options", "Options and Settings"));
-		{
-			if (IPluginManager::Get().FindEnabledPlugin(TEXT("ProjectLauncher")) != nullptr)
+			if (bHasCustomBuilds)
 			{
-				OptionsSection.AddMenuEntry(
-					NAME_None,
-					LOCTEXT("OpenProjectLauncher", "Project Launcher..."),
-					LOCTEXT("OpenProjectLauncher_ToolTip", "Open the Project Launcher for advanced packaging, deploying and launching of your projects"),
-					FSlateIcon(FAppStyle::Get().GetStyleSetName(), "Launcher.TabIcon"),
-					FUIAction(FExecuteAction::CreateStatic(&FTurnkeySupportCallbacks::OpenProjectLauncher))
+				CustomBuildSection.AddSubMenu(
+					"CustomBuildsSubMenu",
+					LOCTEXT("Turnkey_CustomBuilds", "Custom Builds"),
+					LOCTEXT("TurnkeyTooltip_CustomBuilds", "Run custom builds commands defined in config files"),
+					FNewToolMenuDelegate::CreateLambda([this](UToolMenu* SubMenu)
+						{
+							MakeCustomBuildsMenuItems(SubMenu);
+						}
+					),
+					false,
+					FSlateIcon(FAppStyle::GetAppStyleSetName(), "Icons.Advanced")
 				);
 			}
+		}
 
-			OptionsSection.AddMenuEntry(
-				NAME_None,
-				LOCTEXT("OpenLegacyProjectLauncher", "Legacy Project Launcher..."),
-				LOCTEXT("OpenLegacyProjectLauncher_ToolTip", "Open the legacy Project Launcher for advanced packaging, deploying and launching of your projects"),
-				FSlateIcon(FAppStyle::Get().GetStyleSetName(), "Launcher.TabIcon"),
-				FUIAction(FExecuteAction::CreateStatic(&FTurnkeySupportCallbacks::OpenLegacyProjectLauncher))
+		FToolMenuSection& OptionsSection = Menu->AddSection("Options", LOCTEXT("TurnkeySection_Options", "Options and Settings"));
+		{
+			OptionsSection.AddSubMenu("SDKManagementSubMenu",
+				LOCTEXT("SDKManagement", "SDK Management"),
+				LOCTEXT("SDKManagement_ToolTip", "Manage SDK Installations"),
+				FNewToolMenuDelegate::CreateLambda([this](UToolMenu* SubMenu)
+					{
+						MakeSdkManagementMenuItems(SubMenu);	
+					}
+				),
+				false,
+				FSlateIcon(FAppStyle::Get().GetStyleSetName(), "MainFrame.Project")
 			);
 
-
-
 			OptionsSection.AddMenuEntry(
 				NAME_None,
 				LOCTEXT("OpenDeviceManager", "Device Manager..."),
diff --git a/Engine/Source/Editor/TurnkeySupport/Private/TurnkeySupportModule.h b/Engine/Source/Editor/TurnkeySupport/Private/TurnkeySupportModule.h
--- a/Engine/Source/Editor/TurnkeySupport/Private/TurnkeySupportModule.h
+++ b/Engine/Source/Editor/TurnkeySupport/Private/TurnkeySupportModule.h
@@ -23,7 +23,6 @@ public:
 
 	virtual void MakeTurnkeyMenu(FToolMenuSection& MenuSection, const FToolMenuInsert& InsertPosition) const override;
 	virtual void MakeQuickLaunchItems(class UToolMenu* Menu, FOnQuickLaunchSelected ExternalOnClickDelegate) const override;
-	virtual void MakeSimulatorItems(class UToolMenu* Menu, FOnQuickLaunchSelected ExternalOnClickDelegate) const override;
 	virtual void RepeatQuickLaunch(FString DeviceId) override;
 
 
diff --git a/Engine/Source/Editor/TurnkeySupport/Public/Interfaces/ITurnkeySupportModule.h b/Engine/Source/Editor/TurnkeySupport/Public/Interfaces/ITurnkeySupportModule.h
--- a/Engine/Source/Editor/TurnkeySupport/Public/Interfaces/ITurnkeySupportModule.h
+++ b/Engine/Source/Editor/TurnkeySupport/Public/Interfaces/ITurnkeySupportModule.h
@@ -79,11 +79,6 @@ public:
 	 */
 	virtual void MakeQuickLaunchItems(class UToolMenu* Menu, FOnQuickLaunchSelected ExternalOnClickDelegate) const = 0;
 
-	/**
-	 * Make menu items for the Simulator submenu items, so they can be added to the Play menu
-	 */
-	virtual void MakeSimulatorItems(class UToolMenu* Menu, FOnQuickLaunchSelected ExternalOnClickDelegate) const = 0;
-
 	/**
 	 * Runs Turnkey to get the Sdk information for all known platforms
 	 */
diff --git a/Engine/Source/Editor/UnrealEd/Private/Kismet2/DebuggerCommands.cpp b/Engine/Source/Editor/UnrealEd/Private/Kismet2/DebuggerCommands.cpp
--- a/Engine/Source/Editor/UnrealEd/Private/Kismet2/DebuggerCommands.cpp
+++ b/Engine/Source/Editor/UnrealEd/Private/Kismet2/DebuggerCommands.cpp
@@ -80,7 +80,6 @@
 #include "ISettingsEditorModule.h"
 #include "Async/Async.h"
 #include "Misc/FileHelper.h"
-#include "Interfaces/ITurnkeySupportModule.h"
 #include "Settings/ProjectPackagingSettings.h"
 
 
@@ -927,11 +926,6 @@ TSharedRef< SWidget > FPlayWorldCommands::GeneratePlayMenuContent(TSharedRef<FUI
 			}
 		}
 
-#if UE_WITH_TURNKEY_SUPPORT
-		// quick launch on devices
-		ITurnkeySupportModule::Get().MakeQuickLaunchItems(Menu, FOnQuickLaunchSelected::CreateStatic(&RememberQuickLaunch));
-#endif
-
 		// tip section
 		{
 			FToolMenuSection& Section = Menu->AddSection("LevelEditorPlayTip");
diff --git a/Engine/Source/Editor/UnrealEd/Private/PlayLevelViaLauncher.cpp b/Engine/Source/Editor/UnrealEd/Private/PlayLevelViaLauncher.cpp
--- a/Engine/Source/Editor/UnrealEd/Private/PlayLevelViaLauncher.cpp
+++ b/Engine/Source/Editor/UnrealEd/Private/PlayLevelViaLauncher.cpp
@@ -160,7 +160,7 @@ void UEditorEngine::StartPlayUsingLauncherSession(FRequestPlaySessionParams& InR
 		const UProjectPackagingSettings* AllPlatformPackagingSettings = GetDefault<UProjectPackagingSettings>();
 		const UPlatformsMenuSettings* PlatformsSettings = GetDefault<UPlatformsMenuSettings>();
 
-		EProjectPackagingBuildConfigurations BuildConfig = PlatformsSettings->GetBuildConfigurationForPlatform(*IniPlatformName);
+		EProjectPackagingBuildConfigurations BuildConfig = PlatformsSettings->GetLaunchOnBuildConfiguration();
 		// if PPBC_MAX is set, then the project default should be used instead of the per platform build config
 		if (BuildConfig == EProjectPackagingBuildConfigurations::PPBC_MAX)
 		{
