[Editor] Fixes Platform menu not showing custom builds.

If only EngineCustomBuilds was non-empty, the menu would not appear.

#rnx

diff --git a/Engine/Source/Editor/TurnkeySupport/Private/TurnkeySupportModule.cpp b/Engine/Source/Editor/TurnkeySupport/Private/TurnkeySupportModule.cpp
--- a/Engine/Source/Editor/TurnkeySupport/Private/TurnkeySupportModule.cpp
+++ b/Engine/Source/Editor/TurnkeySupport/Private/TurnkeySupportModule.cpp
@@ -3157,7 +3157,7 @@ TSharedRef<SWidget> FTurnkeySupportModule::MakeTurnkeyMenuWidget() const
 			{
 				UProjectPackagingSettings* PackagingSettings = FTurnkeySupportCallbacks::GetPackagingSettingsForPlatform(IniPlatformName);
 
-				if (PackagingSettings->EngineCustomBuilds.IsEmpty() || !PackagingSettings->ProjectCustomBuilds.IsEmpty())
+				if (!PackagingSettings->EngineCustomBuilds.IsEmpty() || !PackagingSettings->ProjectCustomBuilds.IsEmpty())
 				{
 					bHasCustomBuilds = true;
 					break;
