[Editor] Fixes LOCTEXT conflict for TurnkeyTooltip_PrepareForDebugging.

The first usage is the tooltip for the top-level "Prepare for Debugging" entry in the Platforms menu. The second usage, which has been renamed, appears for each platform within the sub-menu and shows a per-platform path in its instructions. So I kept them both.

#rnx

diff --git a/Engine/Source/Editor/TurnkeySupport/Private/TurnkeySupportModule.cpp b/Engine/Source/Editor/TurnkeySupport/Private/TurnkeySupportModule.cpp
--- a/Engine/Source/Editor/TurnkeySupport/Private/TurnkeySupportModule.cpp
+++ b/Engine/Source/Editor/TurnkeySupport/Private/TurnkeySupportModule.cpp
@@ -1918,7 +1918,7 @@ static void MakePrepareForDebuggingMenuItems(UToolMenu* Menu)
 				FToolMenuEntry Entry = FToolMenuEntry::InitMenuEntry(
 					NAME_None,
 					MakeSdkStatusAttribute(IniPlatformName, nullptr),
-					FText::Format(LOCTEXT("TurnkeyTooltip_PrepareForDebugging", "Expects a working remote toolchain and an IPA package with the same name as the project file in the Binaries/{0}/ folder."), FText::FromString(IniPlatformName.ToString().ToUpper())),
+					FText::Format(LOCTEXT("TurnkeyTooltip_PrepareForDebuggingPlatform", "Expects a working remote toolchain and an IPA package with the same name as the project file in the Binaries/{0}/ folder."), FText::FromString(IniPlatformName.ToString().ToUpper())),
 					MakePlatformSdkIconAttribute(IniPlatformName, nullptr),
 					FUIAction(
 						FExecuteAction::CreateStatic(&FTurnkeySupportCallbacks::SetPrepareForDebuggingPlatform, IniPlatformName),
