[CrashReport][Crash] TurnkeySupport!FTurnkeySupportCallbacks::CookOrPackage() [TurnkeySupportModule.cpp]
-Added an ensure and return to more gracefully fail along with some additional logging.  Hopefully we can find out why this happens.
-I found that if one adds C++ to a blueprint project and then packages without exiting the editor to build it and has a windows architecture set for packaging one can induce this failure.  But I'm somewhat doubtful that so many people are doing that.

diff --git a/Engine/Source/Developer/DeveloperToolSettings/Private/ProjectPackagingSettings.cpp b/Engine/Source/Developer/DeveloperToolSettings/Private/ProjectPackagingSettings.cpp
--- a/Engine/Source/Developer/DeveloperToolSettings/Private/ProjectPackagingSettings.cpp
+++ b/Engine/Source/Developer/DeveloperToolSettings/Private/ProjectPackagingSettings.cpp
@@ -242,6 +242,18 @@ const FTargetInfo* FindBestTargetInfo(const FString& TargetName, bool bContentOn
 			DefaultClientTarget = &Target;
 		}
 	}
+
+	// Added in 5.7.2 to help diagnose UE-316569
+	if (DefaultGameTarget == nullptr && DefaultClientTarget == nullptr)
+	{
+		UE_LOG(LogProjectPackagingSettings, Log, TEXT("UProjectPackagingSettings FindBestTargetInfo for '%s' resulted in null FTargetInfo*.  Listing targets that were searched:"), *TargetName);
+		for (const FTargetInfo& Target : TargetsRef)
+		{
+			UE_LOG(LogProjectPackagingSettings, Log, TEXT("  %i %s"), Target.Type , *Target.Name);
+		}
+		UE_LOG(LogProjectPackagingSettings, Log, TEXT("  End of target list."));
+	}
+
 	return (DefaultGameTarget != nullptr) ? DefaultGameTarget : DefaultClientTarget;
 
 }
diff --git a/Engine/Source/Editor/TurnkeySupport/Private/TurnkeySupportModule.cpp b/Engine/Source/Editor/TurnkeySupport/Private/TurnkeySupportModule.cpp
--- a/Engine/Source/Editor/TurnkeySupport/Private/TurnkeySupportModule.cpp
+++ b/Engine/Source/Editor/TurnkeySupport/Private/TurnkeySupportModule.cpp
@@ -408,29 +408,37 @@ public:
 		FString Architecture = PlatformsSettings->GetArchitectureForPlatform(IniPlatformName);
 		if (!Architecture.IsEmpty())
 		{
-			switch (BuildTargetInfo->Type)
+			if (BuildTargetInfo)
 			{
-				case EBuildTargetType::Server:
+				switch (BuildTargetInfo->Type)
 				{
-					BuildCookRunParams += FString::Printf(TEXT(" -serverarchitecture=%s"), *Architecture);
-				}
-				break;
+					case EBuildTargetType::Server:
+					{
+						BuildCookRunParams += FString::Printf(TEXT(" -serverarchitecture=%s"), *Architecture);
+					}
+					break;
 
-				case EBuildTargetType::Editor:
-				{
-					BuildCookRunParams += FString::Printf(TEXT(" -editorarchitecture=%s"), *Architecture);
-				}
-				break;
+					case EBuildTargetType::Editor:
+					{
+						BuildCookRunParams += FString::Printf(TEXT(" -editorarchitecture=%s"), *Architecture);
+					}
+					break;
 
-				default:
-				{
-					BuildCookRunParams += FString::Printf(TEXT(" -clientarchitecture=%s"), *Architecture);
+					default:
+					{
+						BuildCookRunParams += FString::Printf(TEXT(" -clientarchitecture=%s"), *Architecture);
+					}
+					break;
 				}
-				break;
+			}
+			else
+			{
+				// Added in 5.7.2 to help diagnose UE-316569, a crash at the switch statement above because BuildTargetInfo is null.
+				ensureMsgf(BuildTargetInfo, TEXT("CookOrPackage BuildTargetName='%s' IniPlatformName='%s' Mode='%i' BuildCookRunParams='%s 'where BuildTargetInfo is null and architecture '%s' is specified.  One way to do this is to add c++ code to a project, but not build the project editor, then package. Returning without cook or package instead of asserting."), *BuildTargetName, *IniPlatformName.ToString(), Mode, *BuildCookRunParams, *Architecture);
+				return;
 			}
 		}
 
-
 		// per mode settings
 		FText ContentPrepDescription;
 		FText ContentPrepTaskName;
