[Editor] The Platforms menu now always selects a default option on first run.

For the menus that display a list of platforms as radio buttons, if no platform is selected then the button to run the action (e.g., cook content) is disabled. Several users have found this to be confusing, so now a default is always chosen.

The two scenarios where this could occur would be if there was no saved setting, such as the first run of the editor. Once you made a selection, that would be saved and used on the next run. A less common scenario is the saved setting represents a platform that is no longer in a supported state, which could happen if SDKs were uninstalled. This solves both scenarios.

This fixes Cook Content, Package Project, Prepare for Debugging, and Custom Build menus. For Custom Builds, this selects either a default platform or a default device, depending on how the custom build is configured.

diff --git a/Engine/Source/Editor/TurnkeySupport/Private/TurnkeySupportModule.cpp b/Engine/Source/Editor/TurnkeySupport/Private/TurnkeySupportModule.cpp
--- a/Engine/Source/Editor/TurnkeySupport/Private/TurnkeySupportModule.cpp
+++ b/Engine/Source/Editor/TurnkeySupport/Private/TurnkeySupportModule.cpp
@@ -1528,6 +1528,11 @@ static void MakeCookContentMenuItems(UToolMenu* Menu)
 		{
 			ITargetDeviceServicesModule* TargetDeviceServicesModule = static_cast<ITargetDeviceServicesModule*>(FModuleManager::Get().LoadModule(TEXT("TargetDeviceServices")));
 
+			// Always ensure a platform is selected to avoid confusion when the button is disabled. This occurs when
+			// there is no saved setting such as on first launch, or if the saved platform is no longer supported due to
+			// an SDK being uninstalled.
+			bool bUseNextPlatformAsDefault = !FTurnkeySupportCallbacks::CanCookContent();
+				
 			TArray<FName> SupportedPlatforms;
 			GetTurnkeyCategorizedPlatforms(&SupportedPlatforms, nullptr, nullptr);
 
@@ -1597,6 +1602,12 @@ static void MakeCookContentMenuItems(UToolMenu* Menu)
 
 					PlatformsSection.AddEntry(Entry);
 				}
+
+				if (bUseNextPlatformAsDefault)
+				{
+					FTurnkeySupportCallbacks::SetCookPlatform(IniPlatformName);
+					bUseNextPlatformAsDefault = false;
+				}
 			}
 		}));
 	}
@@ -1701,6 +1712,11 @@ static void MakePackageProjectMenuItems(UToolMenu* Menu)
 		{
 			ITargetDeviceServicesModule* TargetDeviceServicesModule = static_cast<ITargetDeviceServicesModule*>(FModuleManager::Get().LoadModule(TEXT("TargetDeviceServices")));
 
+			// Always ensure a platform is selected to avoid confusion when the button is disabled. This occurs when
+			// there is no saved setting such as on first launch, or if the saved platform is no longer supported due to
+			// an SDK being uninstalled.
+			bool bUseNextPlatformAsDefault = !FTurnkeySupportCallbacks::CanPackageProject();
+
 			TArray<FName> SupportedPlatforms;
 			GetTurnkeyCategorizedPlatforms(&SupportedPlatforms, nullptr, nullptr);
 
@@ -1770,6 +1786,12 @@ static void MakePackageProjectMenuItems(UToolMenu* Menu)
 
 					PlatformsSection.AddEntry(Entry);
 				}
+
+				if (bUseNextPlatformAsDefault)
+				{
+					FTurnkeySupportCallbacks::SetPackagePlatform(IniPlatformName);
+					bUseNextPlatformAsDefault = false;
+				}
 			}
 		}));
 	}
@@ -1912,6 +1934,11 @@ static void MakePrepareForDebuggingMenuItems(UToolMenu* Menu)
 	{
 		PlatformsSection.AddDynamicEntry(NAME_None, FNewToolMenuSectionDelegate::CreateLambda([](FToolMenuSection& PlatformsSection)
 		{
+			// Always ensure a platform is selected to avoid confusion when the button is disabled. This occurs when
+			// there is no saved setting such as on first launch, or if the saved platform is no longer supported due to
+			// an SDK being uninstalled.
+			bool bUseNextPlatformAsDefault = !FTurnkeySupportCallbacks::CanPrepareForDebugging();
+
 			TArray<FName> SupportedPlatforms;
 			GetTurnkeyCategorizedPlatforms(&SupportedPlatforms, nullptr, nullptr);
 
@@ -1937,6 +1964,12 @@ static void MakePrepareForDebuggingMenuItems(UToolMenu* Menu)
 				);
 
 				PlatformsSection.AddEntry(Entry);
+
+				if (bUseNextPlatformAsDefault)
+				{
+					FTurnkeySupportCallbacks::SetPrepareForDebuggingPlatform(IniPlatformName);
+					bUseNextPlatformAsDefault = false;
+				}
 			}
 		}));
 	}
@@ -2044,6 +2077,11 @@ static void MakePerCustomBuildMenu(UToolMenu* ToolMenu, ECustomBuildType Type, F
 
 		ITargetDeviceServicesModule* TargetDeviceServicesModule = static_cast<ITargetDeviceServicesModule*>(FModuleManager::Get().LoadModule(TEXT("TargetDeviceServices")));
 
+		// Always ensure a device is selected to avoid confusion when the button is disabled. This occurs when
+		// there is no saved setting such as on first launch, or if the saved platform is no longer supported due to
+		// an SDK being uninstalled.
+		bool bUseNextDeviceAsDefault = !FTurnkeySupportCallbacks::CanExecuteCustomBuild(Type, CustomBuild.Name, bHasDeviceEntry);
+
 		TArray<FName> SupportedPlatforms;
 		GetTurnkeyCategorizedPlatforms(&SupportedPlatforms, nullptr, nullptr);
 
@@ -2071,6 +2109,12 @@ static void MakePerCustomBuildMenu(UToolMenu* ToolMenu, ECustomBuildType Type, F
 						),
 						EUserInterfaceActionType::RadioButton
 					);
+
+					if (bUseNextDeviceAsDefault)
+					{
+						FTurnkeySupportCallbacks::SetDeviceIdForCustomBuild(Type, CustomBuild.Name, IniPlatformName, DeviceId);
+						bUseNextDeviceAsDefault = false;
+					}
 				}
 			}
 		}
@@ -2079,6 +2123,11 @@ static void MakePerCustomBuildMenu(UToolMenu* ToolMenu, ECustomBuildType Type, F
 	{
 		FToolMenuSection& PlatformsSection = ToolMenu->AddSection("PlatformsSection", LOCTEXT("TurnkeySection_Platform", "Platform"));
 		{
+			// Always ensure a platform is selected to avoid confusion when the button is disabled. This occurs when
+			// there is no saved setting such as on first launch, or if the saved platform is no longer supported due to
+			// an SDK being uninstalled.
+			bool bUseNextPlatformAsDefault = !FTurnkeySupportCallbacks::CanExecuteCustomBuild(Type, CustomBuild.Name, bHasDeviceEntry);
+
 			TArray<FName> SupportedPlatforms;
 			GetTurnkeyCategorizedPlatforms(&SupportedPlatforms, nullptr, nullptr);
 
@@ -2151,6 +2200,11 @@ static void MakePerCustomBuildMenu(UToolMenu* ToolMenu, ECustomBuildType Type, F
 						PlatformsSection.AddEntry(Entry);
 					}
 
+					if (bUseNextPlatformAsDefault)
+					{
+						FTurnkeySupportCallbacks::SetPlatformForCustomBuild(Type, CustomBuild.Name, IniPlatformName);
+						bUseNextPlatformAsDefault = false;
+					}
 				}
 			}
 		}
