[Editor] Improves default platform selection in Platforms menu.

On the menus that displayed a list of platforms as radio buttons, the default selection was always the first in the list. A default selection was made any time the saved setting was not supported, most commonly on first run but also if that platform was uninstalled or disabled in the project settings.

These menus display platforms with have compiled-in support and are enabled by the project, regardless of whether the SDK is installed. When a platform is not installed, there is a triangular warning icon next to the platform and if you select it the button is disabled. Choosing the first platform in the list meant possible choosing a platform that isn't supported and the button was disabled. That isn't a great choice for a default. Additionally, after selecting a platform that wasn't installed and re-opening the menu, it reverted back to the default because the logic checked if the platform was supported.

The new behavior is to always respect the user's selection. If the previously selected platform exists in the list, even if it's not installed, the selection won't change. A default is only chosen when the saved setting no longer corresponds to an item in the list, such as the first launch or if that platform was disabled in the project settings.

When chosing a default, it now starts with the platform that is currently running the editor. If that platform is unsupported, then it chooses the first platform in the list that is supported. If no platforms are supported, it selects the last item in the list.

For common scenarios like cooking and packaging, there are no known scenarios where the host platform wouldn't be supported. So the default should always be Windows or Mac. The fallbacks primarily exist for custom builds, which may not support that platform.

Device-specific custom builds always picks the first device in the list as the default. This is more complex than the other cases and there are plans for better long-term solutions to simplify the behavior.

diff --git a/Engine/Source/Editor/TurnkeySupport/Private/TurnkeySupportModule.cpp b/Engine/Source/Editor/TurnkeySupport/Private/TurnkeySupportModule.cpp
--- a/Engine/Source/Editor/TurnkeySupport/Private/TurnkeySupportModule.cpp
+++ b/Engine/Source/Editor/TurnkeySupport/Private/TurnkeySupportModule.cpp
@@ -77,6 +77,16 @@ DEFINE_LOG_CATEGORY(LogTurnkeySupport);
 namespace 
 {
 	FCriticalSection GTurnkeySection;
+
+	// The host platform's IniPlatformName, used to pick defaults.
+#if PLATFORM_WINDOWS
+	const FName GHostIniPlatformName("Windows");
+#elif PLATFORM_MAC
+	const FName GHostIniPlatformName("Mac");
+#else
+	// NAME_None disables this behavior and uses the fallback behavior for choosing a default.
+	const FName GHostIniPlatformName(NAME_None);
+#endif
 }
 
 static FString GetProjectPathForTurnkey()
@@ -627,6 +637,11 @@ public:
 		FTurnkeyEditorSupport::RunUAT(CommandLine, PlatformInfo->DisplayName, ContentPrepDescription, ContentPrepTaskName, ContentPrepIcon, &AnalyticsParamArray);
 	}
 
+	static FName GetPlatformForCustomBuild(ECustomBuildType Type, FString CustomBuildName)
+	{
+		return GetDefault<UPlatformsMenuSettings>()->GetPlatformForCustomBuild(Type, CustomBuildName);
+	}
+
 	static void SetPlatformForCustomBuild(ECustomBuildType Type, FString CustomBuildName, FName IniPlatformName)
 	{
 		UPlatformsMenuSettings* PlatformsSettings = GetMutableDefault<UPlatformsMenuSettings>();
@@ -956,6 +971,11 @@ public:
 		return GetDefault<UCookerSettings>()->bCookOnTheFlyForLaunchOn;
 	}
 
+	static FName GetCookPlatform()
+	{
+		return GetDefault<UPlatformsMenuSettings>()->CookPlatform;
+	}
+
 	static void SetCookPlatform(FName IniPlatformName)
 	{
 		UPlatformsMenuSettings* PlatformsSettings = GetMutableDefault<UPlatformsMenuSettings>();
@@ -991,6 +1011,11 @@ public:
 		CookOrPackage(PlatformsSettings->CookPlatform, PlatformsSettings->CookBuildTarget, EPrepareContentMode::CookOnly);
 	}
 
+	static FName GetPackagePlatform()
+	{
+		return GetDefault<UPlatformsMenuSettings>()->PackagePlatform;
+	}
+
 	static void SetPackagePlatform(FName IniPlatformName)
 	{
 		UPlatformsMenuSettings* PlatformsSettings = GetMutableDefault<UPlatformsMenuSettings>();
@@ -1014,6 +1039,11 @@ public:
 		CookOrPackage(PlatformsSettings->PackagePlatform, PlatformsSettings->PackageBuildTarget, EPrepareContentMode::Package);
 	}
 
+	static FName GetPrepareForDebuggingPlatform()
+	{
+		return GetDefault<UPlatformsMenuSettings>()->PrepareForDebuggingPlatform;
+	}
+
 	static void SetPrepareForDebuggingPlatform(FName IniPlatformName)
 	{
 		UPlatformsMenuSettings* PlatformsSettings = GetMutableDefault<UPlatformsMenuSettings>();
@@ -1528,13 +1558,21 @@ static void MakeCookContentMenuItems(UToolMenu* Menu)
 		{
 			ITargetDeviceServicesModule* TargetDeviceServicesModule = static_cast<ITargetDeviceServicesModule*>(FModuleManager::Get().LoadModule(TEXT("TargetDeviceServices")));
 
+			TArray<FName> SupportedPlatforms;
+			GetTurnkeyCategorizedPlatforms(&SupportedPlatforms, nullptr, nullptr);
+
 			// Always ensure a platform is selected to avoid confusion when the button is disabled. This occurs when
 			// there is no saved setting such as on first launch, or if the saved platform is no longer supported due to
 			// an SDK being uninstalled.
-			bool bUseNextPlatformAsDefault = !FTurnkeySupportCallbacks::CanCookContent();
-				
-			TArray<FName> SupportedPlatforms;
-			GetTurnkeyCategorizedPlatforms(&SupportedPlatforms, nullptr, nullptr);
+			FName SelectedPlatform = FTurnkeySupportCallbacks::GetCookPlatform();
+			bool bUseNextPlatformAsDefault = !SupportedPlatforms.Contains(SelectedPlatform);
+
+			// Use the host platform as the default selection as long as it's in the list and supported.
+			if (bUseNextPlatformAsDefault && SupportedPlatforms.Contains(GHostIniPlatformName))
+			{
+				FTurnkeySupportCallbacks::SetCookPlatform(GHostIniPlatformName);
+				bUseNextPlatformAsDefault = !FTurnkeySupportCallbacks::CanCookContent();
+			}
 
 			for (const FName& IniPlatformName : SupportedPlatforms)
 			{
@@ -1603,10 +1641,12 @@ static void MakeCookContentMenuItems(UToolMenu* Menu)
 					PlatformsSection.AddEntry(Entry);
 				}
 
+				// If the host platform is not supported, choose the first supported platform. If there are none, the
+				// last item in the list ends up selected.
 				if (bUseNextPlatformAsDefault)
 				{
 					FTurnkeySupportCallbacks::SetCookPlatform(IniPlatformName);
-					bUseNextPlatformAsDefault = false;
+					bUseNextPlatformAsDefault = !FTurnkeySupportCallbacks::CanCookContent();
 				}
 			}
 		}));
@@ -1712,13 +1752,21 @@ static void MakePackageProjectMenuItems(UToolMenu* Menu)
 		{
 			ITargetDeviceServicesModule* TargetDeviceServicesModule = static_cast<ITargetDeviceServicesModule*>(FModuleManager::Get().LoadModule(TEXT("TargetDeviceServices")));
 
+			TArray<FName> SupportedPlatforms;
+			GetTurnkeyCategorizedPlatforms(&SupportedPlatforms, nullptr, nullptr);
+
 			// Always ensure a platform is selected to avoid confusion when the button is disabled. This occurs when
 			// there is no saved setting such as on first launch, or if the saved platform is no longer supported due to
 			// an SDK being uninstalled.
-			bool bUseNextPlatformAsDefault = !FTurnkeySupportCallbacks::CanPackageProject();
+			FName SelectedPlatform = FTurnkeySupportCallbacks::GetPackagePlatform();
+			bool bUseNextPlatformAsDefault = !SupportedPlatforms.Contains(SelectedPlatform);
 
-			TArray<FName> SupportedPlatforms;
-			GetTurnkeyCategorizedPlatforms(&SupportedPlatforms, nullptr, nullptr);
+			// Use the host platform as the default selection as long as it's in the list and supported.
+			if (bUseNextPlatformAsDefault && SupportedPlatforms.Contains(GHostIniPlatformName))
+			{
+				FTurnkeySupportCallbacks::SetPackagePlatform(GHostIniPlatformName);
+				bUseNextPlatformAsDefault = !FTurnkeySupportCallbacks::CanPackageProject();
+			}
 
 			for (const FName& IniPlatformName : SupportedPlatforms)
 			{
@@ -1787,10 +1835,12 @@ static void MakePackageProjectMenuItems(UToolMenu* Menu)
 					PlatformsSection.AddEntry(Entry);
 				}
 
+				// If the host platform is not supported, choose the first supported platform. If there are none, the
+				// last item in the list ends up selected.
 				if (bUseNextPlatformAsDefault)
 				{
 					FTurnkeySupportCallbacks::SetPackagePlatform(IniPlatformName);
-					bUseNextPlatformAsDefault = false;
+					bUseNextPlatformAsDefault = !FTurnkeySupportCallbacks::CanPackageProject();
 				}
 			}
 		}));
@@ -1934,13 +1984,14 @@ static void MakePrepareForDebuggingMenuItems(UToolMenu* Menu)
 	{
 		PlatformsSection.AddDynamicEntry(NAME_None, FNewToolMenuSectionDelegate::CreateLambda([](FToolMenuSection& PlatformsSection)
 		{
+			TArray<FName> SupportedPlatforms;
+			GetTurnkeyCategorizedPlatforms(&SupportedPlatforms, nullptr, nullptr);
+
 			// Always ensure a platform is selected to avoid confusion when the button is disabled. This occurs when
 			// there is no saved setting such as on first launch, or if the saved platform is no longer supported due to
 			// an SDK being uninstalled.
-			bool bUseNextPlatformAsDefault = !FTurnkeySupportCallbacks::CanPrepareForDebugging();
-
-			TArray<FName> SupportedPlatforms;
-			GetTurnkeyCategorizedPlatforms(&SupportedPlatforms, nullptr, nullptr);
+			FName SelectedPlatform = FTurnkeySupportCallbacks::GetPrepareForDebuggingPlatform();
+			bool bUseNextPlatformAsDefault = !SupportedPlatforms.Contains(SelectedPlatform) || (SelectedPlatform != TEXT("IOS") && SelectedPlatform != TEXT("TvOS"));
 
 			for (const FName& IniPlatformName : SupportedPlatforms)
 			{
@@ -1965,10 +2016,11 @@ static void MakePrepareForDebuggingMenuItems(UToolMenu* Menu)
 
 				PlatformsSection.AddEntry(Entry);
 
+				// Choose the first supported platform. If there are none, the last item in the list ends up selected.
 				if (bUseNextPlatformAsDefault)
 				{
 					FTurnkeySupportCallbacks::SetPrepareForDebuggingPlatform(IniPlatformName);
-					bUseNextPlatformAsDefault = false;
+					bUseNextPlatformAsDefault = !FTurnkeySupportCallbacks::CanPrepareForDebugging();
 				}
 			}
 		}));
@@ -2077,14 +2129,14 @@ static void MakePerCustomBuildMenu(UToolMenu* ToolMenu, ECustomBuildType Type, F
 
 		ITargetDeviceServicesModule* TargetDeviceServicesModule = static_cast<ITargetDeviceServicesModule*>(FModuleManager::Get().LoadModule(TEXT("TargetDeviceServices")));
 
+		TArray<FName> SupportedPlatforms;
+		GetTurnkeyCategorizedPlatforms(&SupportedPlatforms, nullptr, nullptr);
+
 		// Always ensure a device is selected to avoid confusion when the button is disabled. This occurs when
 		// there is no saved setting such as on first launch, or if the saved platform is no longer supported due to
 		// an SDK being uninstalled.
 		bool bUseNextDeviceAsDefault = !FTurnkeySupportCallbacks::CanExecuteCustomBuild(Type, CustomBuild.Name, bHasDeviceEntry);
 
-		TArray<FName> SupportedPlatforms;
-		GetTurnkeyCategorizedPlatforms(&SupportedPlatforms, nullptr, nullptr);
-
 		for (const FName& IniPlatformName : SupportedPlatforms)
 		{
 			if (CustomBuild.SpecificPlatforms.IsEmpty() || CustomBuild.SpecificPlatforms.Contains(IniPlatformName))
@@ -2110,10 +2162,11 @@ static void MakePerCustomBuildMenu(UToolMenu* ToolMenu, ECustomBuildType Type, F
 						EUserInterfaceActionType::RadioButton
 					);
 
+					// Choose the first valid device as a default.
 					if (bUseNextDeviceAsDefault)
 					{
 						FTurnkeySupportCallbacks::SetDeviceIdForCustomBuild(Type, CustomBuild.Name, IniPlatformName, DeviceId);
-						bUseNextDeviceAsDefault = false;
+						bUseNextDeviceAsDefault = !FTurnkeySupportCallbacks::CanExecuteCustomBuild(Type, CustomBuild.Name, bHasDeviceEntry);
 					}
 				}
 			}
@@ -2123,13 +2176,22 @@ static void MakePerCustomBuildMenu(UToolMenu* ToolMenu, ECustomBuildType Type, F
 	{
 		FToolMenuSection& PlatformsSection = ToolMenu->AddSection("PlatformsSection", LOCTEXT("TurnkeySection_Platform", "Platform"));
 		{
+			TArray<FName> SupportedPlatforms;
+			GetTurnkeyCategorizedPlatforms(&SupportedPlatforms, nullptr, nullptr);
+
 			// Always ensure a platform is selected to avoid confusion when the button is disabled. This occurs when
 			// there is no saved setting such as on first launch, or if the saved platform is no longer supported due to
 			// an SDK being uninstalled.
-			bool bUseNextPlatformAsDefault = !FTurnkeySupportCallbacks::CanExecuteCustomBuild(Type, CustomBuild.Name, bHasDeviceEntry);
+			FName SelectedPlatform = FTurnkeySupportCallbacks::GetPlatformForCustomBuild(Type, CustomBuild.Name);
+			bool bUseNextPlatformAsDefault = !SupportedPlatforms.Contains(SelectedPlatform);
 
-			TArray<FName> SupportedPlatforms;
-			GetTurnkeyCategorizedPlatforms(&SupportedPlatforms, nullptr, nullptr);
+			// Use the host platform as the default selection as long as it's in the list.
+			if (bUseNextPlatformAsDefault && SupportedPlatforms.Contains(GHostIniPlatformName) &&
+				(CustomBuild.SpecificPlatforms.IsEmpty() || CustomBuild.SpecificPlatforms.Contains(GHostIniPlatformName)))
+			{
+				FTurnkeySupportCallbacks::SetPlatformForCustomBuild(Type, CustomBuild.Name, GHostIniPlatformName);
+				bUseNextPlatformAsDefault = !FTurnkeySupportCallbacks::CanExecuteCustomBuild(Type, CustomBuild.Name, bHasDeviceEntry);
+			}
 
 			for (const FName& IniPlatformName : SupportedPlatforms)
 			{
@@ -2200,10 +2262,11 @@ static void MakePerCustomBuildMenu(UToolMenu* ToolMenu, ECustomBuildType Type, F
 						PlatformsSection.AddEntry(Entry);
 					}
 
+					// If the host platform is not supported, choose the first platform.
 					if (bUseNextPlatformAsDefault)
 					{
 						FTurnkeySupportCallbacks::SetPlatformForCustomBuild(Type, CustomBuild.Name, IniPlatformName);
-						bUseNextPlatformAsDefault = false;
+						bUseNextPlatformAsDefault = !FTurnkeySupportCallbacks::CanExecuteCustomBuild(Type, CustomBuild.Name, bHasDeviceEntry);
 					}
 				}
 			}
