Remove Legacy Project Launcher from the main Platforms menu
- still available in platform developer tools
- officially labelled as Deprecated

#rnx

#Synced-CL 50434826

diff --git a/Engine/Plugins/Developer/ProjectLauncher/Source/ProjectLauncher/Private/Widgets/Profiles/SCustomLaunchProfilesPanel.cpp b/Engine/Plugins/Developer/ProjectLauncher/Source/ProjectLauncher/Private/Widgets/Profiles/SCustomLaunchProfilesPanel.cpp
--- a/Engine/Plugins/Developer/ProjectLauncher/Source/ProjectLauncher/Private/Widgets/Profiles/SCustomLaunchProfilesPanel.cpp
+++ b/Engine/Plugins/Developer/ProjectLauncher/Source/ProjectLauncher/Private/Widgets/Profiles/SCustomLaunchProfilesPanel.cpp
@@ -204,8 +204,8 @@ TSharedRef<SWidget> SCustomLaunchProfilesPanel::CreateProfilesPanel()
 			[
 				SNew(SHyperlink)
 				.OnNavigate_Lambda([]() { FGlobalTabmanager::Get()->TryInvokeTab(FTabId("LegacyProjectLauncher")); } )
-				.Text(LOCTEXT("OpenLegacyProjectLauncher","Open Legacy Project Launcher"))
-				.ToolTipText(LOCTEXT("OpenLegacyProjectLauncherTip", "Opens the legacy Project Launcher tab"))
+				.Text(LOCTEXT("OpenLegacyProjectLauncher","Open Legacy Project Launcher (deprecated)"))
+				.ToolTipText(LOCTEXT("OpenLegacyProjectLauncherTip", "Opens the legacy Project Launcher tab (note: this tool is deprecated)"))
 			]
 		];
 	}
diff --git a/Engine/Source/Developer/LegacyProjectLauncher/Private/ProjectLauncherModule.cpp b/Engine/Source/Developer/LegacyProjectLauncher/Private/ProjectLauncherModule.cpp
--- a/Engine/Source/Developer/LegacyProjectLauncher/Private/ProjectLauncherModule.cpp
+++ b/Engine/Source/Developer/LegacyProjectLauncher/Private/ProjectLauncherModule.cpp
@@ -33,8 +33,8 @@ public:
 	virtual void StartupModule() override
 	{
 		FGlobalTabmanager::Get()->RegisterNomadTabSpawner(ProjectLauncherTabName, FOnSpawnTab::CreateRaw(this, &FLegacyProjectLauncherModule::SpawnProjectLauncherTab))
-			.SetDisplayName(NSLOCTEXT("FProjectLauncherModule", "ProjectLauncherTabTitle", "Legacy Project Launcher"))
-			.SetTooltipText(NSLOCTEXT("FProjectLauncherModule", "ProjectLauncherTooltipText", "Open the Legacy Project Launcher tab."))
+			.SetDisplayName(NSLOCTEXT("FProjectLauncherModule", "ProjectLauncherTabTitle", "Legacy Project Launcher (deprecated)"))
+			.SetTooltipText(NSLOCTEXT("FProjectLauncherModule", "ProjectLauncherTooltipText", "Open the Legacy Project Launcher tab (note: this tool is deprecated)."))
 			.SetIcon(FSlateIcon(FAppStyle::GetAppStyleSetName(), "Launcher.TabIcon"))
 			.SetGroup(WorkspaceMenu::GetMenuStructure().GetDeveloperToolsPlatformsCategory());
 	}
diff --git a/Engine/Source/Editor/TurnkeySupport/Private/TurnkeySupportModule.cpp b/Engine/Source/Editor/TurnkeySupport/Private/TurnkeySupportModule.cpp
--- a/Engine/Source/Editor/TurnkeySupport/Private/TurnkeySupportModule.cpp
+++ b/Engine/Source/Editor/TurnkeySupport/Private/TurnkeySupportModule.cpp
@@ -237,11 +237,6 @@ public:
 		return FString::Printf(TEXT("-ReportFilename=\"%s\" -log=\"%s\""), *ReportFilename, *LogFilename);
 	}
 
-	static void OpenLegacyProjectLauncher()
-	{
-		FGlobalTabmanager::Get()->TryInvokeTab(FTabId("LegacyProjectLauncher"));
-	}
-
 	static void OpenProjectLauncher()
 	{
 		FGlobalTabmanager::Get()->TryInvokeTab(FTabId("ProjectLauncher"));
@@ -3206,14 +3201,6 @@ TSharedRef<SWidget> FTurnkeySupportModule::MakeTurnkeyMenuWidget() const
 					FUIAction(FExecuteAction::CreateStatic(&FTurnkeySupportCallbacks::OpenProjectLauncher))
 				);
 			}
-
-			ProjectLauncherSection.AddMenuEntry(
-				NAME_None,
-				LOCTEXT("OpenLegacyProjectLauncher", "Project Launcher (Legacy)..."),
-				LOCTEXT("OpenLegacyProjectLauncher_ToolTip", "Open the legacy Project Launcher for advanced packaging, deploying and launching of your projects"),
-				FSlateIcon(FAppStyle::Get().GetStyleSetName(), "Launcher.TabIcon"),
-				FUIAction(FExecuteAction::CreateStatic(&FTurnkeySupportCallbacks::OpenLegacyProjectLauncher))
-			);
 		}
 
 		FToolMenuSection& ContentManagementSection = Menu->AddSection("ContentManagementSection", LOCTEXT("TurnkeySection_ContentManagement", "Cook or Package"));
